Move ObjectMixer to Engine Module
Publish ObjectMixer

[snapshot CL20876622 - 7/9/2022 5:12:01 PM]

diff --git a/Engine/Source/Editor/LevelEditor/LevelEditor.Build.cs b/Engine/Source/Editor/LevelEditor/LevelEditor.Build.cs
--- a/Engine/Source/Editor/LevelEditor/LevelEditor.Build.cs
+++ b/Engine/Source/Editor/LevelEditor/LevelEditor.Build.cs
@@ -61,6 +61,7 @@ public class LevelEditor : ModuleRules
 				"TypedElementFramework",
 				"TypedElementRuntime",
 				"EngineSettings",
+				"ObjectMixerEditor",
 				"PropertyEditor",
 				"Kismet",
 				"KismetWidgets",
@@ -107,8 +108,7 @@ public class LevelEditor : ModuleRules
 				"VREditor",
                 "Persona",
 				"LevelAssetEditor",
-				"MergeActors",
-
+				"MergeActors"
 			}
 		);
 
diff --git a/Engine/Source/Editor/LevelEditor/Private/SLevelEditor.cpp b/Engine/Source/Editor/LevelEditor/Private/SLevelEditor.cpp
--- a/Engine/Source/Editor/LevelEditor/Private/SLevelEditor.cpp
+++ b/Engine/Source/Editor/LevelEditor/Private/SLevelEditor.cpp
@@ -72,6 +72,7 @@
 #include "EditorViewportLayout.h"
 #include "LevelViewportTabContent.h"
 #include "SLevelViewport.h"
+#include "ObjectMixerEditorModule.h"
 
 #define LOCTEXT_NAMESPACE "SLevelEditor"
 
@@ -1481,6 +1482,11 @@ TSharedRef<SWidget> SLevelEditor::RestoreContentArea( const TSharedRef<SDockTab>
 				.SetIcon(EnvironmentLightingViewerIcon);
 		}
 
+		{
+			// The object mixer module will take care of its tab spawner once it's loaded.
+			FModuleManager::Get().LoadModuleChecked<FObjectMixerEditorModule>( "ObjectMixerEditor" );
+		}
+
 		FTabSpawnerEntry& BuildAndSubmitEntry = LevelEditorTabManager->RegisterTabSpawner(LevelEditorTabIds::LevelEditorBuildAndSubmit, FOnSpawnTab::CreateSP<SLevelEditor, FName, FString>(this, &SLevelEditor::SpawnLevelEditorTab, LevelEditorTabIds::LevelEditorBuildAndSubmit, FString()));
 		BuildAndSubmitEntry.SetAutoGenerateMenuEntry(false);
 
@@ -2172,4 +2178,4 @@ FVector2D SLevelEditor::GetActiveViewportSize()
 	return FVector2D(MaxX, MaxY);
 }
 
-#undef LOCTEXT_NAMESPACE
\ No newline at end of file
+#undef LOCTEXT_NAMESPACE
diff --git a/Engine/Source/Editor/LightMixer/LightMixer.build.cs b/Engine/Source/Editor/LightMixer/LightMixer.build.cs
new file mode 100644
--- /dev/null
+++ b/Engine/Source/Editor/LightMixer/LightMixer.build.cs
@@ -0,0 +1,48 @@
+// Copyright Epic Games, Inc. All Rights Reserved.
+
+using System.IO;
+using UnrealBuildTool;
+
+public class LightMixer : ModuleRules
+{
+	public LightMixer(ReadOnlyTargetRules Target) : base(Target)
+	{
+		PCHUsage = ModuleRules.PCHUsageMode.UseExplicitOrSharedPCHs;
+		
+		PublicIncludePaths.Add(Path.Combine(ModuleDirectory, "Public"));
+
+		PublicDependencyModuleNames.AddRange(
+			new string[]
+			{
+				"Core",
+				"OutputLog",
+				"ObjectMixerEditor",
+				"PropertyEditor"
+			}
+		);
+			
+		
+		PrivateDependencyModuleNames.AddRange(
+			new string[]
+			{
+				"AssetRegistry",
+				"AssetTools",
+				"CoreUObject",
+				"ContentBrowser",
+				"Engine",
+				"EditorStyle",
+				"EditorWidgets",
+				"InputCore",
+				"Kismet",
+				"ObjectMixerEditor",
+				"Projects",
+				"Slate",
+				"SlateCore",
+				"ToolMenus", 
+				"ToolWidgets",
+				"UnrealEd",
+				"WorkspaceMenuStructure"
+			}
+		);
+	}
+}
diff --git a/Engine/Source/Editor/LightMixer/Private/LightMixerModule.cpp b/Engine/Source/Editor/LightMixer/Private/LightMixerModule.cpp
new file mode 100644
--- /dev/null
+++ b/Engine/Source/Editor/LightMixer/Private/LightMixerModule.cpp
@@ -0,0 +1,82 @@
+// Copyright Epic Games, Inc. All Rights Reserved.
+
+#include "LightMixerModule.h"
+
+#include "LightMixerObjectFilter.h"
+#include "LightMixerProjectSettings.h"
+#include "LightMixerStyle.h"
+
+#include "Framework/Docking/TabManager.h"
+#include "ISettingsModule.h"
+#include "ISettingsSection.h"
+#include "Modules/ModuleManager.h"
+#include "ToolMenus.h"
+#include "WorkspaceMenuStructure.h"
+#include "WorkspaceMenuStructureModule.h"
+
+#define LOCTEXT_NAMESPACE "FLightMixerEditorModule"
+
+IMPLEMENT_MODULE(FLightMixerModule, LightMixerEditor)
+
+void FLightMixerModule::StartupModule()
+{
+	FLightMixerStyle::Initialize();
+
+	RegisterTabSpawner();
+	RegisterProjectSettings();
+	
+	TabLabel = LOCTEXT("LightMixerTabLabel", "Light Mixer");
+	
+	DefaultFilterClass = ULightMixerObjectFilter::StaticClass();
+}
+
+void FLightMixerModule::ShutdownModule()
+{
+	UToolMenus::UnregisterOwner(this);
+
+	UnregisterTabSpawner();
+	UnregisterProjectSettings();
+
+	FLightMixerStyle::Shutdown();
+}
+
+void FLightMixerModule::RegisterTabSpawner()
+{
+	FTabSpawnerEntry& BrowserSpawnerEntry = FGlobalTabmanager::Get()->RegisterNomadTabSpawner(
+		"LightMixerToolkit",
+		FOnSpawnTab::CreateRaw(this, & FLightMixerModule::SpawnMainPanelTab))
+			.SetIcon(FSlateIcon(FLightMixerStyle::Get().GetStyleSetName(), "LightMixer.ToolbarButton", "LightMixer.ToolbarButton.Small"))
+			.SetDisplayName(LOCTEXT("OpenLightMixerEditorMenuItem", "Light Mixer"))
+			.SetTooltipText(LOCTEXT("OpenLightMixerEditorTooltip", "Open Light Mixer"))
+			.SetMenuType(ETabSpawnerMenuType::Enabled);
+
+	BrowserSpawnerEntry.SetGroup(WorkspaceMenu::GetMenuStructure().GetLevelEditorCategory());
+}
+
+void FLightMixerModule::UnregisterTabSpawner()
+{
+	FGlobalTabmanager::Get()->UnregisterNomadTabSpawner(FObjectMixerEditorModule::ObjectMixerToolkitPanelTabId);
+}
+
+void FLightMixerModule::RegisterProjectSettings() const
+{
+	if (ISettingsModule* SettingsModule = FModuleManager::GetModulePtr<ISettingsModule>("Settings"))
+	{
+		// User Project Settings
+		const TSharedPtr<ISettingsSection> ProjectSettingsSectionPtr = SettingsModule->RegisterSettings(
+			"Project", "Plugins", "Light Mixer",
+			LOCTEXT("LightMixerSettingsCategoryDisplayName", "Light Mixer"),
+			LOCTEXT("LightMixerSettingsDescription", "Configure Light Mixer user settings"),
+			GetMutableDefault<ULightMixerProjectSettings>());
+	}
+}
+
+void FLightMixerModule::UnregisterProjectSettings() const
+{
+	if (ISettingsModule* SettingsModule = FModuleManager::GetModulePtr<ISettingsModule>("Settings"))
+	{
+		SettingsModule->UnregisterSettings("Project", "Plugins", "Light Mixer");
+	}
+}
+
+#undef LOCTEXT_NAMESPACE
diff --git a/Engine/Source/Editor/LightMixer/Private/LightMixerStyle.cpp b/Engine/Source/Editor/LightMixer/Private/LightMixerStyle.cpp
new file mode 100644
--- /dev/null
+++ b/Engine/Source/Editor/LightMixer/Private/LightMixerStyle.cpp
@@ -0,0 +1,90 @@
+// Copyright Epic Games, Inc. All Rights Reserved.
+
+#include "LightMixerStyle.h"
+
+#include "Styling/AppStyle.h"
+#include "Framework/Application/SlateApplication.h"
+#include "Interfaces/IPluginManager.h"
+#include "Styling/CoreStyle.h"
+#include "Styling/SlateStyleRegistry.h"
+#include "Styling/SlateTypes.h"
+#include "Styling/StyleColors.h"
+
+TSharedPtr<FSlateStyleSet> FLightMixerStyle::StyleInstance = nullptr;
+
+void FLightMixerStyle::Initialize()
+{
+	if (!StyleInstance.IsValid())
+	{
+		StyleInstance = Create();
+		FSlateStyleRegistry::RegisterSlateStyle(*StyleInstance);
+	}
+}
+
+void FLightMixerStyle::Shutdown()
+{
+	FSlateStyleRegistry::UnRegisterSlateStyle(*StyleInstance);
+	ensure(StyleInstance.IsUnique());
+	StyleInstance.Reset();
+}
+
+void FLightMixerStyle::ReloadTextures()
+{
+	FSlateApplication::Get().GetRenderer()->ReloadTextureResources();
+}
+
+const ISlateStyle& FLightMixerStyle::Get()
+{
+	return *StyleInstance;
+}
+
+const FName& FLightMixerStyle::GetStyleSetName() const
+{
+	static FName StyleName(TEXT("LightMixer"));
+	return StyleName;
+}
+
+const FSlateBrush* FLightMixerStyle::GetBrush(const FName PropertyName, const ANSICHAR* Specifier, const ISlateStyle* RequestingStyle) const
+{
+	return StyleInstance->GetBrush(PropertyName, Specifier);
+}
+
+#define IMAGE_BRUSH( RelativePath, ... ) FSlateImageBrush( Style->RootToContentDir( RelativePath, TEXT(".png") ), __VA_ARGS__ )
+#define BOX_BRUSH( RelativePath, ... ) FSlateBoxBrush( Style->RootToContentDir( RelativePath, TEXT(".png") ), __VA_ARGS__ )
+#define BORDER_BRUSH( RelativePath, ... ) FSlateBorderBrush( Style->RootToContentDir( RelativePath, TEXT(".png") ), __VA_ARGS__ )
+#define IMAGE_BRUSH_SVG( RelativePath, ... ) FSlateVectorImageBrush( Style->RootToContentDir(RelativePath, TEXT(".svg") ), __VA_ARGS__)
+#define IMAGE_PLUGIN_BRUSH_SVG( PluginName, RelativePath, ... ) FSlateVectorImageBrush( FLightMixerEditorStyle::GetExternalPluginContent(PluginName, RelativePath, ".svg"), __VA_ARGS__)
+
+const FVector2D Icon64x64(64.f, 64.f);
+const FVector2D Icon40x40(40.0f, 40.0f);
+const FVector2D Icon20x20(20.0f, 20.0f);
+const FVector2D Icon16x16(16.0f, 16.0f);
+const FVector2D Icon12x12(12.0f, 12.0f);
+const FVector2D Icon8x8(8.f, 8.f);
+
+FString FLightMixerStyle::GetExternalPluginContent(const FString& PluginName, const FString& RelativePath, const ANSICHAR* Extension)
+{
+	FString ContentDir = IPluginManager::Get().FindPlugin(PluginName)->GetBaseDir() / RelativePath + Extension;
+	return ContentDir;
+}
+
+TSharedRef< FSlateStyleSet > FLightMixerStyle::Create()
+{
+	TSharedRef<FSlateStyleSet> Style = MakeShared<FSlateStyleSet>("LightMixer");
+	
+	//Style->SetContentRoot(FPaths::Combine(Plugin->GetBaseDir(), TEXT("Resources")));
+
+	Style->SetCoreContentRoot(FPaths::EngineContentDir() / TEXT("Slate"));
+
+	// Icons
+	Style->Set("LightMixer.ToolbarButton", new IMAGE_BRUSH_SVG("Icons/LightMixer", Icon40x40));
+	Style->Set("LightMixer.ToolbarButton.Small", new IMAGE_BRUSH_SVG("Icons/LightMixer", Icon20x20));
+	
+	return Style;
+}
+
+#undef IMAGE_BRUSH
+#undef BOX_BRUSH
+#undef BORDER_BRUSH
+#undef IMAGE_BRUSH_SVG
+#undef IMAGE_PLUGIN_BRUSH_SVG
diff --git a/Engine/Source/Editor/LightMixer/Private/LightMixerStyle.h b/Engine/Source/Editor/LightMixer/Private/LightMixerStyle.h
new file mode 100644
--- /dev/null
+++ b/Engine/Source/Editor/LightMixer/Private/LightMixerStyle.h
@@ -0,0 +1,36 @@
+// Copyright Epic Games, Inc. All Rights Reserved.
+
+#pragma once
+
+#include "Styling/SlateStyle.h"
+
+class FLightMixerStyle : FSlateStyleSet
+{
+public:
+
+	static void Initialize();
+
+	static void Shutdown();
+
+	static void ReloadTextures();
+
+	static const ISlateStyle& Get();
+
+	virtual const FName& GetStyleSetName() const override;
+
+	virtual const FSlateBrush* GetBrush(const FName PropertyName, const ANSICHAR* Specifier = nullptr, const ISlateStyle* RequestingStyle = nullptr) const override;
+
+	template< typename WidgetStyleType >
+	static const WidgetStyleType& GetWidgetStyle(FName PropertyName, const ANSICHAR* Specifier = nullptr)
+	{
+		return StyleInstance->GetWidgetStyle<WidgetStyleType>(PropertyName, Specifier);
+	}
+
+private:
+
+	static FString GetExternalPluginContent(const FString& PluginName, const FString& RelativePath, const ANSICHAR* Extension);
+	
+	static TSharedRef<FSlateStyleSet> Create();
+
+	static TSharedPtr<FSlateStyleSet> StyleInstance;
+};
diff --git a/Engine/Source/Editor/LightMixer/Public/LightMixerModule.h b/Engine/Source/Editor/LightMixer/Public/LightMixerModule.h
new file mode 100644
--- /dev/null
+++ b/Engine/Source/Editor/LightMixer/Public/LightMixerModule.h
@@ -0,0 +1,23 @@
+// Copyright Epic Games, Inc. All Rights Reserved.
+
+#pragma once
+
+#include "ObjectMixerEditorModule.h"
+
+class LIGHTMIXER_API FLightMixerModule : public FObjectMixerEditorModule
+{
+public:
+
+	//~ Begin IModuleInterface Interface
+	virtual void StartupModule() override;
+	virtual void ShutdownModule() override;
+	//~ End IModuleInterface Interface
+
+	void RegisterTabSpawner();
+
+	//~ Begin FObjectMixerEditorModule overrides
+	virtual void UnregisterTabSpawner() override;
+	virtual void RegisterProjectSettings() const override;
+	virtual void UnregisterProjectSettings() const override;
+	//~ End FObjectMixerEditorModule overrides
+};
diff --git a/Engine/Source/Editor/LightMixer/Public/LightMixerObjectFilter.h b/Engine/Source/Editor/LightMixer/Public/LightMixerObjectFilter.h
new file mode 100644
--- /dev/null
+++ b/Engine/Source/Editor/LightMixer/Public/LightMixerObjectFilter.h
@@ -0,0 +1,62 @@
+// Copyright Epic Games, Inc. All Rights Reserved.
+
+#pragma once
+
+#include "ObjectFilter/ObjectMixerEditorObjectFilter.h"
+
+#include "GameFramework/Actor.h"
+
+#include "LightMixerObjectFilter.generated.h"
+
+UCLASS(BlueprintType, EditInlineNew)
+class LIGHTMIXER_API ULightMixerObjectFilter : public UObjectMixerObjectFilter
+{
+	GENERATED_BODY()
+public:
+	
+	virtual TArray<UClass*> GetObjectClassesToFilter() const override
+	{
+		return
+		{
+			ULightComponent::StaticClass()
+		};
+	}
+
+	virtual FText GetRowDisplayName(UObject* InObject) const override
+	{
+		if (InObject)
+		{
+			if (const AActor* Outer = InObject->GetTypedOuter<AActor>())
+			{
+				return FText::Format(INVTEXT("{0} ({1})"), FText::FromString(Outer->GetActorLabel()), Super::GetRowDisplayName(InObject));
+			}
+		}
+
+		return Super::GetRowDisplayName(InObject);
+	}
+
+	virtual TArray<FName> GetColumnsToShowByDefault() const override
+	{
+		const TArray<FName> IncludeList =
+		{
+			"Intensity", "LightColor", "AttenuationRadius"
+		};
+
+		return IncludeList;
+	}
+
+	virtual TArray<FName> GetForceAddedColumns() const override
+	{
+		return {"LightColor"};
+	}
+
+	virtual bool ShouldIncludeUnsupportedProperties() const override
+	{
+		return true;
+	}
+
+	virtual EObjectMixerPropertyInheritanceInclusionOptions GetObjectMixerPropertyInheritanceInclusionOptions() const override
+	{
+		return EObjectMixerPropertyInheritanceInclusionOptions::IncludeAllParentsAndChildren;
+	}
+};
diff --git a/Engine/Source/Editor/LightMixer/Public/LightMixerProjectSettings.h b/Engine/Source/Editor/LightMixer/Public/LightMixerProjectSettings.h
new file mode 100644
--- /dev/null
+++ b/Engine/Source/Editor/LightMixer/Public/LightMixerProjectSettings.h
@@ -0,0 +1,22 @@
+// Copyright Epic Games, Inc. All Rights Reserved.
+
+#pragma once
+
+#include "LightMixerProjectSettings.generated.h"
+
+UCLASS(config = Engine, defaultconfig)
+class LIGHTMIXER_API ULightMixerProjectSettings : public UObject
+{
+	GENERATED_BODY()
+public:
+	
+	ULightMixerProjectSettings(const FObjectInitializer& ObjectInitializer)
+	{}
+
+	/**
+	 * If true, the Object Mixer menu item will be disabled and removed.
+	 * This is useful if you only want to use Light Mixer and avoid cluttering the menus.
+	 */
+	UPROPERTY(Config, EditAnywhere, Category="Light Mixer", meta = (ConfigRestartRequired = true))
+	bool bHideObjectMixerMenuItem = false;
+};
diff --git a/Engine/Source/Editor/ObjectMixer/ObjectMixerEditor.build.cs b/Engine/Source/Editor/ObjectMixer/ObjectMixerEditor.build.cs
new file mode 100644
--- /dev/null
+++ b/Engine/Source/Editor/ObjectMixer/ObjectMixerEditor.build.cs
@@ -0,0 +1,48 @@
+// Copyright Epic Games, Inc. All Rights Reserved.
+
+using System.IO;
+using UnrealBuildTool;
+
+public class ObjectMixerEditor : ModuleRules
+{
+	public ObjectMixerEditor(ReadOnlyTargetRules Target) : base(Target)
+	{
+		PCHUsage = ModuleRules.PCHUsageMode.UseExplicitOrSharedPCHs;
+		
+		PublicIncludePaths.Add(Path.Combine(ModuleDirectory, "Public"));
+
+		PublicDependencyModuleNames.AddRange(
+			new string[]
+			{
+				"Core",
+				"OutputLog",
+				"PropertyEditor"
+			}
+		);
+			
+		
+		PrivateDependencyModuleNames.AddRange(
+			new string[]
+			{
+				"AssetRegistry",
+				"AssetTools",
+				"CoreUObject",
+				"ContentBrowser",
+				"Engine",
+				"EditorStyle",
+				"EditorWidgets",
+				"InputCore",
+				"Kismet",
+				"LevelEditor",
+				"LightMixer",
+				"Projects",
+				"Slate",
+				"SlateCore",
+				"ToolMenus", 
+				"ToolWidgets",
+				"UnrealEd",
+				"WorkspaceMenuStructure"
+			}
+		);
+	}
+}
diff --git a/Engine/Source/Editor/ObjectMixer/Private/ObjectFilter/ObjectMixerEditorObjectFilter.cpp b/Engine/Source/Editor/ObjectMixer/Private/ObjectFilter/ObjectMixerEditorObjectFilter.cpp
new file mode 100644
--- /dev/null
+++ b/Engine/Source/Editor/ObjectMixer/Private/ObjectFilter/ObjectMixerEditorObjectFilter.cpp
@@ -0,0 +1,170 @@
+// Copyright Epic Games, Inc. All Rights Reserved.
+
+#pragma once
+
+#include "ObjectFilter/ObjectMixerEditorObjectFilter.h"
+
+#include "GameFramework/Actor.h"
+
+FText UObjectMixerObjectFilter::GetRowDisplayName(UObject* InObject) const
+{
+	if (InObject->IsValidLowLevel())
+	{
+		return FText::FromString(InObject->GetName());
+	}
+
+	return FText::GetEmpty();
+}
+
+bool UObjectMixerObjectFilter::GetRowEditorVisibility(UObject* InObject) const
+{
+	if (InObject->IsValidLowLevel())
+	{
+		TObjectPtr<AActor> Actor = Cast<AActor>(InObject);
+	
+		if (!Actor)
+		{
+			Actor = InObject->GetTypedOuter<AActor>();
+		}
+
+		return Actor ? !Actor->IsTemporarilyHiddenInEditor() : false;
+	}
+
+	return false;
+}
+
+void UObjectMixerObjectFilter::OnSetRowEditorVisibility(UObject* InObject, bool bNewIsVisible) const
+{
+	if (InObject->IsValidLowLevel())
+	{
+		TObjectPtr<AActor> Actor = Cast<AActor>(InObject);
+	
+		if (!Actor)
+		{
+			Actor = InObject->GetTypedOuter<AActor>();
+		}
+
+		if (Actor)
+		{
+			Actor->SetIsTemporarilyHiddenInEditor(!bNewIsVisible);
+		}
+	}
+}
+
+TArray<FName> UObjectMixerObjectFilter::GetColumnsToShowByDefault() const
+{
+	return {};
+}
+
+TArray<FName> UObjectMixerObjectFilter::GetColumnsFilter() const
+{
+	return {};
+}
+
+EObjectMixerPropertyInheritanceInclusionOptions UObjectMixerObjectFilter::GetObjectMixerPropertyInheritanceInclusionOptions() const
+{
+	return EObjectMixerPropertyInheritanceInclusionOptions::None;
+}
+
+bool UObjectMixerObjectFilter::ShouldIncludeUnsupportedProperties() const
+{
+	return false;
+}
+
+TArray<UClass*> UObjectMixerObjectFilter::GetParentAndChildClassesFromSpecifiedClasses(
+	const TArray<UClass*>& InSpecifiedClasses, EObjectMixerPropertyInheritanceInclusionOptions Options)
+{
+	// 'None' means we only want the specified classes
+	if (Options == EObjectMixerPropertyInheritanceInclusionOptions::None)
+	{
+		return InSpecifiedClasses;
+	}
+	
+	TArray<UClass*> ReturnValue;
+
+	auto GetChildClassesLambda = [](const UClass* Class, TArray<UClass*>& OutReturnValue, const bool bRecursive)
+	{
+		TArray<UClass*> DerivedClasses;
+		GetDerivedClasses(Class, DerivedClasses, bRecursive);
+		OutReturnValue.Append(DerivedClasses);
+	};
+
+	for (UClass* Class : InSpecifiedClasses)
+	{
+		ReturnValue.Add(Class);
+
+		// Super Classes
+
+			// Immediate only
+		if (Options == EObjectMixerPropertyInheritanceInclusionOptions::IncludeOnlyImmediateParent ||
+			Options == EObjectMixerPropertyInheritanceInclusionOptions::IncludeOnlyImmediateParentAndChildren ||
+			Options == EObjectMixerPropertyInheritanceInclusionOptions::IncludeOnlyImmediateParentAndAllChildren)
+		{
+			if (UClass* Super = Class->GetSuperClass())
+			{
+				ReturnValue.Add(Super);
+			}
+		}
+			// All Parents
+		if (Options == EObjectMixerPropertyInheritanceInclusionOptions::IncludeAllParents ||
+			Options == EObjectMixerPropertyInheritanceInclusionOptions::IncludeAllParentsAndChildren ||
+			Options == EObjectMixerPropertyInheritanceInclusionOptions::IncludeAllParentsAndOnlyImmediateChildren)
+		{
+			if (UClass* Super = Class->GetSuperClass())
+			{
+				while (Super)
+				{
+					ReturnValue.Add(Super);
+					Super = Super->GetSuperClass();
+				}
+			}
+		}
+
+		// Child Classes
+
+			// Immediate only
+		if (Options == EObjectMixerPropertyInheritanceInclusionOptions::IncludeOnlyImmediateChildren ||
+			Options == EObjectMixerPropertyInheritanceInclusionOptions::IncludeOnlyImmediateParentAndChildren ||
+			Options == EObjectMixerPropertyInheritanceInclusionOptions::IncludeAllParentsAndOnlyImmediateChildren)
+		{
+			GetChildClassesLambda(Class, ReturnValue, false);
+		}
+			// All Children
+		if (Options == EObjectMixerPropertyInheritanceInclusionOptions::IncludeAllChildren ||
+			Options == EObjectMixerPropertyInheritanceInclusionOptions::IncludeAllParentsAndChildren ||
+			Options == EObjectMixerPropertyInheritanceInclusionOptions::IncludeOnlyImmediateParentAndAllChildren)
+		{
+			GetChildClassesLambda(Class, ReturnValue, true);
+		}
+	}
+
+	return ReturnValue;
+}
+
+EFieldIterationFlags UObjectMixerObjectFilter::GetDesiredFieldIterationFlags(const bool bIncludeInheritedProperties)
+{
+	return bIncludeInheritedProperties ? EFieldIterationFlags::IncludeSuper : EFieldIterationFlags::Default;
+}
+
+TArray<FName> UObjectMixerObjectFilter::GenerateIncludeListFromExcludeList(const TSet<FName>& ExcludeList) const
+{
+	TSet<FName> IncludeList;
+
+	const EObjectMixerPropertyInheritanceInclusionOptions Options =
+		GetObjectMixerPropertyInheritanceInclusionOptions();
+	TArray<UClass*> SpecifiedClasses =
+		GetParentAndChildClassesFromSpecifiedClasses(GetObjectClassesToFilter(), Options);
+	
+	for (const UClass* Class : SpecifiedClasses)
+	{
+		for (TFieldIterator<FProperty> FieldIterator(Class); FieldIterator; ++FieldIterator)
+		{
+			if (const FProperty* Property = *FieldIterator)
+			{
+				IncludeList.Add(Property->GetFName());
+			}
+		}
+	}
+
+	return IncludeList.Difference(ExcludeList).Array();
+}
diff --git a/Engine/Source/Editor/ObjectMixer/Private/ObjectMixerEditorModule.cpp b/Engine/Source/Editor/ObjectMixer/Private/ObjectMixerEditorModule.cpp
new file mode 100644
--- /dev/null
+++ b/Engine/Source/Editor/ObjectMixer/Private/ObjectMixerEditorModule.cpp
@@ -0,0 +1,149 @@
+// Copyright Epic Games, Inc. All Rights Reserved.
+
+#include "ObjectMixerEditorModule.h"
+
+#include "ObjectMixerEditorLog.h"
+#include "ObjectMixerEditorProjectSettings.h"
+#include "ObjectMixerEditorStyle.h"
+#include "Views/MainPanel/ObjectMixerEditorMainPanel.h"
+
+#include "Framework/Docking/TabManager.h"
+#include "ISettingsModule.h"
+#include "Modules/ModuleManager.h"
+#include "WorkspaceMenuStructure.h"
+#include "WorkspaceMenuStructureModule.h"
+
+#define LOCTEXT_NAMESPACE "FObjectMixerEditorModule"
+
+IMPLEMENT_MODULE(FObjectMixerEditorModule, ObjectMixerEditor)
+
+const FName FObjectMixerEditorModule::ObjectMixerToolkitPanelTabId(TEXT("ObjectMixerToolkit"));
+
+void FObjectMixerEditorModule::StartupModule()
+{
+	FObjectMixerEditorStyle::Initialize();
+
+	RegisterMenuGroupAndTabSpawner();
+	RegisterProjectSettings();
+	
+	TabLabel = LOCTEXT("ObjectMixerTabLabel", "Object Mixer");
+
+	// Initialize Light Mixer
+	FModuleManager::Get().LoadModuleChecked("LightMixer");
+}
+
+void FObjectMixerEditorModule::ShutdownModule()
+{
+	FObjectMixerEditorStyle::Shutdown();
+	
+	MainPanel.Reset();
+
+	UnregisterTabSpawner();
+	UnregisterProjectSettings();
+}
+
+FObjectMixerEditorModule& FObjectMixerEditorModule::Get()
+{
+	return FModuleManager::LoadModuleChecked< FObjectMixerEditorModule >("ObjectMixerEditor");
+}
+
+TSharedPtr<SWidget> FObjectMixerEditorModule::MakeObjectMixerDialog() const
+{
+	if (MainPanel.IsValid())
+	{
+		return MainPanel->GetOrCreateWidget();
+	}
+
+	return nullptr;
+}
+
+void FObjectMixerEditorModule::RebuildList(const FString InItemToScrollTo, bool bShouldCacheValues) const
+{
+	MainPanel->RebuildList(InItemToScrollTo, bShouldCacheValues);
+}
+
+void FObjectMixerEditorModule::RefreshList() const
+{
+	MainPanel->RefreshList();
+}
+
+void FObjectMixerEditorModule::RegisterMenuGroupAndTabSpawner()
+{
+	const FText MenuItemName = LOCTEXT("OpenObjectMixerEditorMenuItem", "Object Mixer");
+	const FSlateIcon MenuItemIcon =
+		FSlateIcon(FObjectMixerEditorStyle::Get().GetStyleSetName(), "ObjectMixer.ToolbarButton", "ObjectMixer.ToolbarButton.Small");
+	
+	FTabSpawnerEntry& BrowserSpawnerEntry = FGlobalTabmanager::Get()->RegisterNomadTabSpawner(
+		FObjectMixerEditorModule::ObjectMixerToolkitPanelTabId,
+		FOnSpawnTab::CreateRaw(&FObjectMixerEditorModule::Get(), & FObjectMixerEditorModule::SpawnMainPanelTab))
+		.SetIcon(MenuItemIcon)
+		.SetDisplayName(MenuItemName)
+		.SetTooltipText(LOCTEXT("OpenObjectMixerEditorMenuItem", "Open an Object Mixer instance."))
+		.SetMenuType(ETabSpawnerMenuType::Hidden);
+
+	WorkspaceGroup = WorkspaceMenu::GetMenuStructure().GetLevelEditorCategory()->AddGroup(MenuItemName, MenuItemIcon);
+
+	BrowserSpawnerEntry.SetGroup(WorkspaceGroup.ToSharedRef());
+}
+
+bool FObjectMixerEditorModule::RegisterItemInMenuGroup(FWorkspaceItem& InItem)
+{
+	if (WorkspaceGroup)
+	{
+		WorkspaceGroup->AddItem(MakeShareable(&InItem));
+		
+		return true;
+	}
+
+	return false;
+}
+
+void FObjectMixerEditorModule::UnregisterTabSpawner()
+{
+	FGlobalTabmanager::Get()->UnregisterNomadTabSpawner(FObjectMixerEditorModule::ObjectMixerToolkitPanelTabId);
+}
+
+void FObjectMixerEditorModule::RegisterProjectSettings() const
+{
+	if (ISettingsModule* SettingsModule = FModuleManager::GetModulePtr<ISettingsModule>("Settings"))
+	{
+		// User Project Settings
+		const TSharedPtr<ISettingsSection> ProjectSettingsSectionPtr = SettingsModule->RegisterSettings(
+			"Project", "Editor", "Object Mixer",
+			LOCTEXT("ObjectMixerSettingsCategoryDisplayName", "Object Mixer"),
+			LOCTEXT("ObjectMixerSettingsDescription", "Configure Object Mixer user settings"),
+			GetMutableDefault<UObjectMixerEditorProjectSettings>());	
+	}
+}
+
+void FObjectMixerEditorModule::UnregisterProjectSettings() const
+{
+	if (ISettingsModule* SettingsModule = FModuleManager::GetModulePtr<ISettingsModule>("Settings"))
+	{
+		SettingsModule->UnregisterSettings("Project", "Plugins", "Object Mixer");
+	}
+}
+
+TSharedRef<SDockTab> FObjectMixerEditorModule::SpawnMainPanelTab(const FSpawnTabArgs& Args)
+{
+	MainPanel = MakeShared<FObjectMixerEditorMainPanel>();
+	
+	const TSharedRef<SDockTab> DockTab =
+		SNew(SDockTab)
+		.Label(TabLabel)
+		.TabRole(ETabRole::NomadTab)
+	;
+	const TSharedPtr<SWidget> ObjectMixerDialog = MakeObjectMixerDialog();
+	DockTab->SetContent(ObjectMixerDialog ? ObjectMixerDialog.ToSharedRef() : SNullWidget::NullWidget);
+	MainPanel->OnClassSelectionChanged(DefaultFilterClass);
+	MainPanel->RebuildList();
+			
+	return DockTab;
+}
+
+TSharedPtr<FWorkspaceItem> FObjectMixerEditorModule::GetWorkspaceGroup()
+{
+	return WorkspaceGroup;
+}
+
+#undef LOCTEXT_NAMESPACE
diff --git a/Engine/Source/Editor/ObjectMixer/Private/ObjectMixerEditorStyle.cpp b/Engine/Source/Editor/ObjectMixer/Private/ObjectMixerEditorStyle.cpp
new file mode 100644
--- /dev/null
+++ b/Engine/Source/Editor/ObjectMixer/Private/ObjectMixerEditorStyle.cpp
@@ -0,0 +1,103 @@
+// Copyright Epic Games, Inc. All Rights Reserved.
+
+#include "ObjectMixerEditorStyle.h"
+
+#include "Styling/AppStyle.h"
+#include "Framework/Application/SlateApplication.h"
+#include "Interfaces/IPluginManager.h"
+#include "Styling/CoreStyle.h"
+#include "Styling/SlateStyleRegistry.h"
+#include "Styling/SlateTypes.h"
+#include "Styling/StyleColors.h"
+
+TSharedPtr<FSlateStyleSet> FObjectMixerEditorStyle::StyleInstance = nullptr;
+
+void FObjectMixerEditorStyle::Initialize()
+{
+	if (!StyleInstance.IsValid())
+	{
+		StyleInstance = Create();
+		FSlateStyleRegistry::RegisterSlateStyle(*StyleInstance);
+	}
+}
+
+void FObjectMixerEditorStyle::Shutdown()
+{
+	FSlateStyleRegistry::UnRegisterSlateStyle(*StyleInstance);
+	ensure(StyleInstance.IsUnique());
+	StyleInstance.Reset();
+}
+
+void FObjectMixerEditorStyle::ReloadTextures()
+{
+	FSlateApplication::Get().GetRenderer()->ReloadTextureResources();
+}
+
+const ISlateStyle& FObjectMixerEditorStyle::Get()
+{
+	return *StyleInstance;
+}
+
+const FName& FObjectMixerEditorStyle::GetStyleSetName() const
+{
+	static FName StyleName(TEXT("ObjectMixerEditor"));
+	return StyleName;
+}
+
+FString GetConcertContent(const FString& RelativePath, const ANSICHAR* Extension)
+{
+	static FString ContentDir = IPluginManager::Get().FindPlugin(TEXT("ConcertSyncClient"))->GetContentDir();
+	return (ContentDir / RelativePath) + Extension;
+}
+
+const FSlateBrush* FObjectMixerEditorStyle::GetBrush(const FName PropertyName, const ANSICHAR* Specifier, const ISlateStyle* RequestingStyle) const
+{
+	return StyleInstance->GetBrush(PropertyName, Specifier);
+}
+
+#define IMAGE_BRUSH( RelativePath, ... ) FSlateImageBrush( Style->RootToContentDir( RelativePath, TEXT(".png") ), __VA_ARGS__ )
+#define BOX_BRUSH( RelativePath, ... ) FSlateBoxBrush( Style->RootToContentDir( RelativePath, TEXT(".png") ), __VA_ARGS__ )
+#define BORDER_BRUSH( RelativePath, ... ) FSlateBorderBrush( Style->RootToContentDir( RelativePath, TEXT(".png") ), __VA_ARGS__ )
+#define IMAGE_BRUSH_SVG( RelativePath, ... ) FSlateVectorImageBrush( Style->RootToContentDir(RelativePath, TEXT(".svg") ), __VA_ARGS__)
+#define IMAGE_PLUGIN_BRUSH_SVG( PluginName, RelativePath, ... ) FSlateVectorImageBrush( FObjectMixerEditorStyle::GetExternalPluginContent(PluginName, RelativePath, ".svg"), __VA_ARGS__)
+
+const FVector2D Icon64x64(64.f, 64.f);
+const FVector2D Icon40x40(40.0f, 40.0f);
+const FVector2D Icon20x20(20.0f, 20.0f);
+const FVector2D Icon16x16(16.0f, 16.0f);
+const FVector2D Icon12x12(12.0f, 12.0f);
+const FVector2D Icon8x8(8.f, 8.f);
+
+FString FObjectMixerEditorStyle::GetExternalPluginContent(const FString& PluginName, const FString& RelativePath, const ANSICHAR* Extension)
+{
+	FString ContentDir = IPluginManager::Get().FindPlugin(PluginName)->GetBaseDir() / RelativePath + Extension;
+	return ContentDir;
+}
+
+TSharedRef< FSlateStyleSet > FObjectMixerEditorStyle::Create()
+{
+	TSharedRef<FSlateStyleSet> Style = MakeShared<FSlateStyleSet>("ObjectMixerEditor");
+	
+	//Style->SetContentRoot(FPaths::Combine(Plugin->GetBaseDir(), TEXT("Resources")));
+
+	//Style->SetCoreContentRoot(FPaths::EngineContentDir() / TEXT("Slate"));
+
+	// Icons
+	Style->Set("ObjectMixer.ToolbarButton", new IMAGE_BRUSH_SVG("Icons/ObjectMixer", Icon40x40));
+	Style->Set("ObjectMixer.ToolbarButton.Small", new IMAGE_BRUSH_SVG("Icons/ObjectMixer", Icon20x20));
+
+	// Brush
+	Style->Set("ObjectMixerEditor.BrightBorder", new FSlateColorBrush(FStyleColors::Header));
+
+	// Border colors for Results view
+	Style->Set("ObjectMixerEditor.HeaderRowBorder", new FSlateColorBrush(FStyleColors::Black));
+	Style->Set("ObjectMixerEditor.DefaultBorder", new FSlateColorBrush(FStyleColors::Transparent));
+	
+	return Style;
+}
+
+#undef IMAGE_BRUSH
+#undef BOX_BRUSH
+#undef BORDER_BRUSH
+#undef IMAGE_BRUSH_SVG
+#undef IMAGE_PLUGIN_BRUSH_SVG
diff --git a/Engine/Source/Editor/ObjectMixer/Private/ObjectMixerFilterFactory.cpp b/Engine/Source/Editor/ObjectMixer/Private/ObjectMixerFilterFactory.cpp
new file mode 100644
--- /dev/null
+++ b/Engine/Source/Editor/ObjectMixer/Private/ObjectMixerFilterFactory.cpp
@@ -0,0 +1,48 @@
+// Copyright Epic Games, Inc. All Rights Reserved.
+
+#include "ObjectMixerFilterFactory.h"
+
+#include "ObjectFilter/ObjectMixerEditorObjectFilter.h"
+
+#include "AssetToolsModule.h"
+#include "Kismet2/KismetEditorUtilities.h"
+
+#define LOCTEXT_NAMESPACE "ObjectMixerFilterFactory"
+
+UObjectMixerBlueprintFilterFactory::UObjectMixerBlueprintFilterFactory()
+{
+	SupportedClass = UBlueprint::StaticClass();
+	ParentClass = UObjectMixerBlueprintObjectFilter::StaticClass();
+
+	bCreateNew = true;
+	bEditAfterNew = true;
+}
+
+FText UObjectMixerBlueprintFilterFactory::GetDisplayName() const
+{
+	return LOCTEXT("DisplayName", "Object Mixer Filter");
+}
+
+UObject* UObjectMixerBlueprintFilterFactory::FactoryCreateNew(UClass* Class, UObject* InParent, FName Name, EObjectFlags Flags, UObject* Context, FFeedbackContext* Warn)
+{
+	UBlueprint* FilterBlueprint = nullptr;
+	if (ensure(SupportedClass == Class))
+	{
+		ensure(0 != (RF_Public & Flags));
+		FilterBlueprint = FKismetEditorUtilities::CreateBlueprint(ParentClass, InParent, Name, BPTYPE_Normal, UBlueprint::StaticClass(), UBlueprintGeneratedClass::StaticClass(), NAME_None);
+	}
+	return FilterBlueprint;	
+}
+
+FString UObjectMixerBlueprintFilterFactory::GetDefaultNewAssetName() const
+{
+	return "NewObjectMixerFilter";
+}
+
+uint32 UObjectMixerBlueprintFilterFactory::GetMenuCategories() const
+{
+	IAssetTools& AssetTools = FModuleManager::LoadModuleChecked<FAssetToolsModule>("AssetTools").Get();
+	return AssetTools.RegisterAdvancedAssetCategory("Object Mixer", LOCTEXT("AssetCategoryName", "Object Mixer"));
+}
+
+#undef LOCTEXT_NAMESPACE
diff --git a/Engine/Source/Editor/ObjectMixer/Private/Views/List/ObjectMixerEditorList.cpp b/Engine/Source/Editor/ObjectMixer/Private/Views/List/ObjectMixerEditorList.cpp
new file mode 100644
--- /dev/null
+++ b/Engine/Source/Editor/ObjectMixer/Private/Views/List/ObjectMixerEditorList.cpp
@@ -0,0 +1,135 @@
+// Copyright Epic Games, Inc. All Rights Reserved.
+
+#include "Views/List/ObjectMixerEditorList.h"
+
+#include "Views/List/SObjectMixerEditorList.h"
+
+#include "Editor.h"
+
+FObjectMixerEditorList::FObjectMixerEditorList()
+{
+	EditorDelegateHandles.Add(FEditorDelegates::MapChange.AddLambda([this](uint32 Index)
+	{
+		RebuildList();
+	}));
+	EditorDelegateHandles.Add(FEditorDelegates::OnMapOpened.AddLambda([this](const FString& Filename, bool bAsTemplate)
+	{
+		RebuildList();
+	}));
+	EditorDelegateHandles.Add(FEditorDelegates::ActorPropertiesChange.AddLambda([this]()
+	{
+		RebuildList();
+	}));
+	EditorDelegateHandles.Add(FEditorDelegates::OnDeleteActorsEnd.AddLambda([this]()
+	{
+		RebuildList();
+	}));
+	EditorDelegateHandles.Add(FEditorDelegates::OnApplyObjectToActor.AddLambda([this](UObject* Object, AActor* Actor)
+	{
+		RebuildList();
+	}));
+	EditorDelegateHandles.Add(FEditorDelegates::OnDuplicateActorsEnd.AddLambda([this]()
+	{
+		RebuildList();
+	}));
+	EditorDelegateHandles.Add(FEditorDelegates::OnNewActorsDropped.AddLambda([this](const TArray<UObject*>& Objects, const TArray<AActor*>& Actors)
+	{
+		RebuildList();
+	}));
+	EditorDelegateHandles.Add(FEditorDelegates::OnNewActorsPlaced.AddLambda([this](UObject* Object, const TArray<AActor*>& Actors)
+	{
+		RebuildList();
+	}));
+}
+
+FObjectMixerEditorList::~FObjectMixerEditorList()
+{
+	ListWidget.Reset();
+
+	// Unbind Delegates
+	FEditorDelegates::MapChange.RemoveAll(this);
+	FEditorDelegates::OnMapOpened.RemoveAll(this);
+	FEditorDelegates::ActorPropertiesChange.RemoveAll(this);
+	FEditorDelegates::OnDeleteActorsEnd.RemoveAll(this);
+	FEditorDelegates::OnApplyObjectToActor.RemoveAll(this);
+	FEditorDelegates::OnDuplicateActorsEnd.RemoveAll(this);
+	FEditorDelegates::OnNewActorsDropped.RemoveAll(this);
+	FEditorDelegates::OnNewActorsPlaced.RemoveAll(this);
+
+	for (FDelegateHandle Delegate : EditorDelegateHandles)
+	{
+		Delegate.Reset();
+	}
+	EditorDelegateHandles.Empty();
+
+	if (ObjectFilterPtr.IsValid())
+	{
+		ObjectFilterPtr.Get()->RemoveFromRoot();
+	}
+}
+
+TSharedRef<SWidget> FObjectMixerEditorList::GetOrCreateWidget()
+{
+	if (!ListWidget.IsValid())
+	{
+		SAssignNew(ListWidget, SObjectMixerEditorList, SharedThis(this));
+	}
+
+	RebuildList();
+
+	return ListWidget.ToSharedRef();
+}
+
+UObjectMixerObjectFilter* FObjectMixerEditorList::GetObjectFilter()
+{
+	if (!ObjectFilterPtr.IsValid())
+	{
+		CacheObjectFilterObject();
+	}
+
+	return ObjectFilterPtr.Get();
+}
+
+void FObjectMixerEditorList::CacheObjectFilterObject()
+{
+	if (ObjectFilterPtr.IsValid())
+	{
+		ObjectFilterPtr->RemoveFromRoot();
+		ObjectFilterPtr = nullptr;
+	}
+	
+	if (const UClass* Class = GetObjectFilterClass())
+	{
+		ObjectFilterPtr = NewObject<UObjectMixerObjectFilter>(GetTransientPackage(), Class);
+		ObjectFilterPtr.Get()->AddToRoot();
+	}
+}
+
+void FObjectMixerEditorList::SetSearchString(const FString& SearchString)
+{
+	ListWidget->SetSearchStringInSearchInputField(SearchString);
+}
+
+void FObjectMixerEditorList::ClearList() const
+{
+	if (ListWidget.IsValid())
+	{
+		ListWidget->ClearList();
+	}
+}
+
+void FObjectMixerEditorList::RebuildList(const FString& InItemToScrollTo) const
+{
+	if (ListWidget.IsValid())
+	{
+		ListWidget->RebuildList(InItemToScrollTo);
+	}
+}
+
+void FObjectMixerEditorList::RefreshList() const
+{
+	if (ListWidget.IsValid())
+	{
+		ListWidget->RefreshList();
+	}
+}
diff --git a/Engine/Source/Editor/ObjectMixer/Private/Views/List/ObjectMixerEditorListRow.cpp b/Engine/Source/Editor/ObjectMixer/Private/Views/List/ObjectMixerEditorListRow.cpp
new file mode 100644
--- /dev/null
+++ b/Engine/Source/Editor/ObjectMixer/Private/Views/List/ObjectMixerEditorListRow.cpp
@@ -0,0 +1,326 @@
+// Copyright Epic Games, Inc. All Rights Reserved.
+
+#include "Views/List/ObjectMixerEditorListRow.h"
+
+#include "Algo/AllOf.h"
+#include "ClassIconFinder.h"
+#include "GameFramework/Actor.h"
+#include "Styling/SlateIconFinder.h"
+#include "Views/List/ObjectMixerEditorList.h"
+#include "Views/List/SObjectMixerEditorList.h"
+
+FObjectMixerEditorListRow::~FObjectMixerEditorListRow()
+{
+	FlushReferences();
+}
+
+void FObjectMixerEditorListRow::FlushReferences()
+{
+	if (ChildRows.Num())
+	{
+		ChildRows.Empty();
+	}
+}
+
+UObjectMixerObjectFilter* FObjectMixerEditorListRow::GetObjectFilter() const
+{
+	if (const TSharedPtr<SObjectMixerEditorList> PinnedListView = GetListViewPtr().Pin())
+	{
+		if (const TSharedPtr<FObjectMixerEditorList> PinnedListModel = PinnedListView->GetListModelPtr().Pin())
+		{
+			if (UObjectMixerObjectFilter* Filter = PinnedListModel->GetObjectFilter())
+			{
+				return Filter;
+			}
+		}
+	}
+
+	return nullptr;
+}
+
+FObjectMixerEditorListRow::EObjectMixerEditorListRowType FObjectMixerEditorListRow::GetRowType() const
+{
+	return RowType;
+}
+
+int32 FObjectMixerEditorListRow::GetChildDepth() const
+{
+	return ChildDepth;
+}
+
+void FObjectMixerEditorListRow::SetChildDepth(const int32 InDepth)
+{
+	ChildDepth = InDepth;
+}
+
+int32 FObjectMixerEditorListRow::GetSortOrder() const
+{
+	return SortOrder;
+}
+
+void FObjectMixerEditorListRow::SetSortOrder(const int32 InNewOrder)
+{
+	SortOrder = InNewOrder;
+}
+
+TWeakPtr<FObjectMixerEditorListRow> FObjectMixerEditorListRow::GetDirectParentRow() const
+{
+	return DirectParentRow;
+}
+
+void FObjectMixerEditorListRow::SetDirectParentRow(
+	const TWeakPtr<FObjectMixerEditorListRow>& InDirectParentRow)
+{
+	DirectParentRow = InDirectParentRow;
+}
+
+const TArray<FObjectMixerEditorListRowPtr>& FObjectMixerEditorListRow::GetChildRows() const
+{
+	return ChildRows;
+}
+
+int32 FObjectMixerEditorListRow::GetChildCount() const
+{
+	return ChildRows.Num();
+}
+
+void FObjectMixerEditorListRow::SetChildRows(const TArray<FObjectMixerEditorListRowPtr>& InChildRows)
+{
+	ChildRows = InChildRows;
+}
+
+void FObjectMixerEditorListRow::AddToChildRows(const FObjectMixerEditorListRowPtr& InRow)
+{
+	ChildRows.Add(InRow);
+}
+
+void FObjectMixerEditorListRow::InsertChildRowAtIndex(const FObjectMixerEditorListRowPtr& InRow,
+                                                           const int32 AtIndex)
+{
+	ChildRows.Insert(InRow, AtIndex);
+}
+
+bool FObjectMixerEditorListRow::GetIsTreeViewItemExpanded() const
+{
+	return bIsTreeViewItemExpanded;
+}
+
+void FObjectMixerEditorListRow::SetIsTreeViewItemExpanded(const bool bNewExpanded)
+{
+	bIsTreeViewItemExpanded = bNewExpanded;
+}
+
+bool FObjectMixerEditorListRow::GetShouldFlashOnScrollIntoView() const
+{
+	return bShouldFlashOnScrollIntoView;
+}
+
+void FObjectMixerEditorListRow::SetShouldFlashOnScrollIntoView(const bool bNewShouldFlashOnScrollIntoView)
+{
+	bShouldFlashOnScrollIntoView = bNewShouldFlashOnScrollIntoView;
+}
+
+bool FObjectMixerEditorListRow::GetShouldExpandAllChildren() const
+{
+	return bShouldExpandAllChildren;
+}
+
+void FObjectMixerEditorListRow::SetShouldExpandAllChildren(const bool bNewShouldExpandAllChildren)
+{
+	bShouldExpandAllChildren = bNewShouldExpandAllChildren;
+}
+
+void FObjectMixerEditorListRow::ResetToStartupValueAndSource() const
+{
+
+}
+
+bool FObjectMixerEditorListRow::MatchSearchTokensToSearchTerms(
+	const TArray<FString> InTokens, ESearchCase::Type InSearchCase)
+{
+	// If the search is cleared we'll consider the row to pass search
+	bool bMatchFound = true;
+
+	if (const TObjectPtr<UObject> Object = GetObject())
+	{
+		if (CachedSearchTerms.IsEmpty())
+		{
+			CachedSearchTerms = Object->GetName();
+
+			if (const UObjectMixerObjectFilter* Filter = GetObjectFilter())
+			{
+				CachedSearchTerms = Filter->GetRowDisplayName(Object).ToString();
+			}
+		}
+
+		// Match any
+		for (const FString& Token : InTokens)
+		{
+			// Match all of these
+			const FString SpaceDelimiter = " ";
+			TArray<FString> OutSpacedArray;
+			if (Token.Contains(SpaceDelimiter) && Token.ParseIntoArray(OutSpacedArray, *SpaceDelimiter, true) > 1)
+			{
+				bMatchFound = Algo::AllOf(OutSpacedArray, [this, InSearchCase](const FString& Comparator)
+				{
+					return CachedSearchTerms.Contains(Comparator, InSearchCase);
+				});
+			}
+			else
+			{
+				bMatchFound = CachedSearchTerms.Contains(Token, InSearchCase);
+			}
+
+			if (bMatchFound)
+			{
+				break;
+			}
+		}
+	}
+
+	bDoesRowMatchSearchTerms = bMatchFound;
+
+	return bMatchFound;
+}
+
+void FObjectMixerEditorListRow::ExecuteSearchOnChildNodes(const FString& SearchString) const
+{
+	TArray<FString> Tokens;
+
+	SearchString.ParseIntoArray(Tokens, TEXT(" "), true);
+
+	ExecuteSearchOnChildNodes(Tokens);
+}
+
+void FObjectMixerEditorListRow::ExecuteSearchOnChildNodes(const TArray<FString>& Tokens) const
+{
+	for (const FObjectMixerEditorListRowPtr& ChildRow : GetChildRows())
+	{
+		if (!ensure(ChildRow.IsValid()))
+		{
+			continue;
+		}
+
+		if (ChildRow->GetRowType() == EObjectMixerEditorListRowType::Group)
+		{
+			if (ChildRow->MatchSearchTokensToSearchTerms(Tokens))
+			{
+				// If the group name matches then we pass an empty string to search child nodes since we want them all to be visible
+				ChildRow->ExecuteSearchOnChildNodes("");
+			}
+			else
+			{
+				// Otherwise we iterate over all child nodes to determine which should and should not be visible
+				ChildRow->ExecuteSearchOnChildNodes(Tokens);
+			}
+		}
+		else
+		{
+			ChildRow->MatchSearchTokensToSearchTerms(Tokens);
+		}
+	}
+}
+
+bool FObjectMixerEditorListRow::GetDoesRowPassFilters() const
+{
+	return bDoesRowPassFilters;
+}
+
+void FObjectMixerEditorListRow::SetDoesRowPassFilters(const bool bPass)
+{
+	bDoesRowPassFilters = bPass;
+}
+
+bool FObjectMixerEditorListRow::GetIsSelected()
+{
+	check (ListViewPtr.IsValid());
+
+	return ListViewPtr.Pin()->IsTreeViewItemSelected(SharedThis(this));
+}
+
+bool FObjectMixerEditorListRow::ShouldBeVisible() const
+{
+	return (bDoesRowMatchSearchTerms && bDoesRowPassFilters) || HasVisibleChildren();
+}
+
+EVisibility FObjectMixerEditorListRow::GetDesiredVisibility() const
+{
+	return ShouldBeVisible() ? EVisibility::Visible : EVisibility::Collapsed;
+}
+
+TArray<FObjectMixerEditorListRowPtr> FObjectMixerEditorListRow::GetSelectedTreeViewItems() const
+{
+	return ListViewPtr.Pin()->GetSelectedTreeViewItems();
+}
+
+const FSlateBrush* FObjectMixerEditorListRow::GetObjectIconBrush()
+{
+	if (GetRowType() != SingleItem)
+	{
+		return nullptr;
+	}
+	
+	if (UObject* RowObject = GetObject())
+	{
+		auto GetIconForActor = [](AActor* InActor)
+		{
+			FName IconName = InActor->GetCustomIconName();
+			if (IconName == NAME_None)
+			{
+				IconName = InActor->GetClass()->GetFName();
+			}
+
+			return FClassIconFinder::FindIconForActor(InActor);
+		};
+
+		if (AActor* AsActor = Cast<AActor>(RowObject))
+		{
+			return GetIconForActor(AsActor);
+		}
+	
+		if (RowObject->IsA(UActorComponent::StaticClass()))
+		{
+			if (AActor* OuterActor = RowObject->GetTypedOuter<AActor>())
+			{
+				return GetIconForActor(OuterActor);
+			}
+		
+			return FSlateIconFinder::FindIconBrushForClass(RowObject->GetClass(), TEXT("SCS.Component"));
+		}
+	}
+
+	return nullptr;
+}
+
+bool FObjectMixerEditorListRow::GetObjectVisibility()
+{
+	if (const UObjectMixerObjectFilter* Filter = GetObjectFilter())
+	{
+		return Filter->GetRowEditorVisibility(GetObject());
+	}
+
+	return false;
+}
+
+void FObjectMixerEditorListRow::SetObjectVisibility(const bool bNewIsVisible)
+{
+	if (const UObjectMixerObjectFilter* Filter = GetObjectFilter())
+	{
+		Filter->OnSetRowEditorVisibility(GetObject(), bNewIsVisible);
+	}
+}
+
+bool FObjectMixerEditorListRow::IsThisRowSolo() const
+{
+	return GetListViewPtr().Pin()->GetSoloRow().HasSameObject(this);
+}
+
+void FObjectMixerEditorListRow::SetThisAsSoloRow()
+{
+	GetListViewPtr().Pin()->SetSoloRow(SharedThis(this));
+}
+
+void FObjectMixerEditorListRow::ClearSoloRow()
+{
+	GetListViewPtr().Pin()->ClearSoloRow();
+}
diff --git a/Engine/Source/Editor/ObjectMixer/Private/Views/List/SObjectMixerEditorList.cpp b/Engine/Source/Editor/ObjectMixer/Private/Views/List/SObjectMixerEditorList.cpp
new file mode 100644
--- /dev/null
+++ b/Engine/Source/Editor/ObjectMixer/Private/Views/List/SObjectMixerEditorList.cpp
@@ -0,0 +1,952 @@
+// Copyright Epic Games, Inc. All Rights Reserved.
+
+#include "Views/List/SObjectMixerEditorList.h"
+
+// Filters
+#include "Views/List/ObjectMixerEditorListFilters/ObjectMixerEditorListFilter_Source.h"
+
+#include "ObjectMixerEditorLog.h"
+#include "ObjectMixerEditorModule.h"
+#include "ObjectMixerEditorStyle.h"
+#include "Views/List/SObjectMixerEditorListRow.h"
+
+#include "Algo/AllOf.h"
+#include "Algo/AnyOf.h"
+#include "Algo/Find.h"
+#include "Components/LightComponent.h"
+#include "Editor.h"
+#include "Framework/MultiBox/MultiBoxBuilder.h"
+#include "ObjectFilter/ObjectMixerEditorObjectFilter.h"
+#include "Styling/StyleColors.h"
+#include "UObject/UObjectIterator.h"
+#include "Widgets/Input/SCheckBox.h"
+#include "Widgets/Input/SSearchBox.h"
+#include "Widgets/Layout/SScrollBox.h"
+#include "Widgets/Layout/SWidgetSwitcher.h"
+#include "Widgets/Text/SRichTextBlock.h"
+
+#define LOCTEXT_NAMESPACE "ObjectMixerEditorList"
+
+const FName SObjectMixerEditorList::ItemNameColumnName(TEXT("Builtin_Name"));
+const FName SObjectMixerEditorList::EditorVisibilityColumnName(TEXT("Builtin_EditorVisibility"));
+const FName SObjectMixerEditorList::EditorVisibilitySoloColumnName(TEXT("Builtin_EditorVisibilitySolo"));
+
+void SObjectMixerEditorList::Construct(const FArguments& InArgs, TSharedRef<FObjectMixerEditorList> ListModel)
+{
+	ListModelPtr = ListModel;
+	
+	// Set Default Sorting info
+	ActiveSortingType = EColumnSortMode::Ascending;
+	
+	HeaderRow = SNew(SHeaderRow)
+				.CanSelectGeneratedColumn(false)
+				.Visibility(EVisibility::Visible)
+				;
+
+	SetupFilters();
+	
+	ChildSlot
+	[
+		SNew(SVerticalBox)
+		
+		+ SVerticalBox::Slot()
+		.VAlign(VAlign_Top)
+		.AutoHeight()
+		.Padding(FMargin(8.f, 0.f, 8.f, 0.f))
+		[
+			SNew(SHorizontalBox)
+
+			+SHorizontalBox::Slot()
+			.HAlign(HAlign_Fill)
+			.VAlign(VAlign_Center)
+			.Padding(0.f, 1.f, 0.f, 1.f)
+			[
+				SAssignNew(ListSearchBoxPtr, SSearchBox)
+				.HintText(LOCTEXT("SearchHintText", "Search Scene Objects"))
+				.ToolTipText(LOCTEXT("ObjectMixerEditorList_TooltipText", "Search Scene Objects"))
+				.OnTextChanged_Raw(this, &SObjectMixerEditorList::OnListViewSearchTextChanged)
+			]
+
+			// Show Options
+			+SHorizontalBox::Slot()
+			.HAlign(HAlign_Right)
+			.VAlign(VAlign_Center)
+			.AutoWidth()
+			.Padding(8.f, 1.f, 0.f, 1.f)
+			[
+				SAssignNew( ViewOptionsComboButton, SComboButton )
+				.ContentPadding(4.f)
+				.ToolTipText(LOCTEXT("ShowOptions_Tooltip", "Show options to affect the visibility of items in the Object Mixer list"))
+				.ComboButtonStyle( FAppStyle::Get(), "SimpleComboButtonWithIcon" ) // Use the tool bar item style for this button
+				.OnGetMenuContent( this, &SObjectMixerEditorList::BuildShowOptionsMenu)
+				.HasDownArrow(false)
+				.ButtonContent()
+				[
+					SNew(SImage)
+					.ColorAndOpacity(FSlateColor::UseForeground())
+					.Image( FAppStyle::Get().GetBrush("Icons.Settings") )
+				]
+			]
+		]
+
+		+ SVerticalBox::Slot()
+		.HAlign(HAlign_Fill)
+		.VAlign(VAlign_Fill)
+		[
+			SNew(SWidgetSwitcher)
+			.WidgetIndex_Lambda([this]()
+			{
+				return DoesTreeViewHaveVisibleChildren() ? 0 : 1;
+			})
+			
+			+ SWidgetSwitcher::Slot()
+			.HAlign(HAlign_Fill)
+			.VAlign(VAlign_Fill)
+			.Padding(2.0f, 2.0f, 2.0f, 2.0f)
+			[
+				SAssignNew(TreeViewPtr, STreeView<FObjectMixerEditorListRowPtr>)
+				.HeaderRow(HeaderRow)
+				.SelectionMode(ESelectionMode::Multi)
+				.OnSelectionChanged_Lambda([this] (const FObjectMixerEditorListRowPtr& Row, const ESelectInfo::Type SelectionType)
+				{
+					if (GEditor && FSlateApplication::Get().GetModifierKeys().IsAltDown())
+					{
+						GEditor->SelectNone(false, true, true);
+
+						for (const TSharedPtr<FObjectMixerEditorListRow>& SelectedRow : TreeViewPtr->GetSelectedItems())
+						{
+							if (SelectedRow->GetRowType() == FObjectMixerEditorListRow::SingleItem)
+							{
+								AActor* Actor = Cast<AActor>(SelectedRow->GetObject());
+
+								if (!Actor)
+								{
+									Actor = SelectedRow->GetObject()->GetTypedOuter<AActor>();
+								}
+
+								if (Actor)
+								{
+									GEditor->SelectActor( Actor, true, true, true );
+								}
+							}
+						}
+					}
+				})
+				.TreeItemsSource(&VisibleTreeViewObjects)
+				.OnGenerateRow_Lambda([this](FObjectMixerEditorListRowPtr Row, const TSharedRef<STableViewBase>& OwnerTable)
+					{
+						check(Row.IsValid());
+					
+						return SNew(SObjectMixerEditorListRow, TreeViewPtr.ToSharedRef(), Row)
+								.Visibility_Raw(Row.Get(), &FObjectMixerEditorListRow::GetDesiredVisibility);
+					})
+				.OnGetChildren_Raw(this, &SObjectMixerEditorList::OnGetRowChildren)
+				.OnExpansionChanged_Raw(this, &SObjectMixerEditorList::OnRowChildExpansionChange, false)
+				.OnSetExpansionRecursive(this, &SObjectMixerEditorList::OnRowChildExpansionChange, true)
+			]
+
+			// For when no rows exist in view
+			+ SWidgetSwitcher::Slot()
+			.HAlign(HAlign_Fill)
+			.Padding(2.0f, 24.0f, 2.0f, 2.0f)
+			[
+				SNew(SRichTextBlock)
+				.DecoratorStyleSet(&FAppStyle::Get())
+				.AutoWrapText(true)
+				.Justification(ETextJustify::Center)
+				.Text_Lambda([this]()
+				{
+					// Preset Empty List (with filter)
+                    return LOCTEXT("EmptyListPresetWithFilter", "No matching items in your list.\n\nCheck your filters.");
+				})
+			]
+		]
+	];
+}
+
+SObjectMixerEditorList::~SObjectMixerEditorList()
+{	
+	HeaderRow.Reset();
+	
+	ListSearchBoxPtr.Reset();
+	ViewOptionsComboButton.Reset();
+	ListBoxContainerPtr.Reset();
+
+	FlushMemory(false);
+
+	ShowFilters.Reset();
+	TreeViewPtr.Reset();
+}
+
+void SObjectMixerEditorList::RebuildList(const FString& InItemToScrollTo)
+{
+	GenerateTreeView();
+
+	if (!InItemToScrollTo.IsEmpty())
+	{
+		FObjectMixerEditorListRowPtr ScrollToItem = nullptr;
+
+		for (const FObjectMixerEditorListRowPtr& Item : TreeViewRootObjects)
+		{
+			if (false)
+			{
+				ScrollToItem = Item;
+				break;
+			}
+		}
+
+		if (ScrollToItem.IsValid())
+		{
+			ScrollToItem->SetShouldFlashOnScrollIntoView(true);
+			TreeViewPtr->RequestScrollIntoView(ScrollToItem);
+		}
+	}
+}
+
+void SObjectMixerEditorList::RefreshList()
+{
+	if (TreeViewRootObjects.Num() > 0)
+	{
+		// Apply last search
+		ExecuteListViewSearchOnAllRows(GetSearchStringFromSearchInputField(), false);
+
+		// Enforce Sort
+		const FName& SortingName = GetActiveSortingColumnName();
+		ExecuteSort(SortingName, GetSortModeForColumn(SortingName), false);
+
+		// Show/Hide rows based on SetBy changes and filter settings
+		EvaluateIfRowsPassFilters(false);
+	}
+
+	FindVisibleObjectsAndRequestTreeRefresh();
+}
+
+TArray<FObjectMixerEditorListRowPtr> SObjectMixerEditorList::GetSelectedTreeViewItems() const
+{
+	return TreeViewPtr->GetSelectedItems();
+}
+
+int32 SObjectMixerEditorList::GetSelectedTreeViewItemCount() const
+{
+	return TreeViewPtr->GetSelectedItems().Num();
+}
+
+void SObjectMixerEditorList::SetSelectedTreeViewItemActorsEditorVisible(const bool bNewIsVisible)
+{
+	for (TSharedPtr<FObjectMixerEditorListRow> Item : TreeViewPtr->GetSelectedItems())
+	{
+		Item->SetObjectVisibility(bNewIsVisible);
+	}
+}
+
+bool SObjectMixerEditorList::IsTreeViewItemSelected(TSharedRef<FObjectMixerEditorListRow> Item)
+{
+	return TreeViewPtr->GetSelectedItems().Contains(Item);
+}
+
+TArray<FObjectMixerEditorListRowPtr> SObjectMixerEditorList::GetTreeViewItems() const
+{
+	return TreeViewRootObjects;
+}
+
+void SObjectMixerEditorList::SetTreeViewItems(const TArray<FObjectMixerEditorListRowPtr>& InItems)
+{
+	TreeViewRootObjects = InItems;
+
+	TreeViewPtr->RequestListRefresh();
+}
+
+FString SObjectMixerEditorList::GetSearchStringFromSearchInputField() const
+{
+	return ensureAlwaysMsgf(ListSearchBoxPtr.IsValid(),
+		TEXT("%hs: ListSearchBoxPtr is not valid. Check to make sure it was created."), __FUNCTION__)
+	? ListSearchBoxPtr->GetText().ToString() : "";
+}
+
+void SObjectMixerEditorList::SetSearchStringInSearchInputField(const FString InSearchString) const
+{
+	if (ensureAlwaysMsgf(ListSearchBoxPtr.IsValid(),
+		TEXT("%hs: ListSearchBoxPtr is not valid. Check to make sure it was created."), __FUNCTION__))
+	{
+		ListSearchBoxPtr->SetText(FText::FromString(InSearchString));
+	}
+}
+
+void SObjectMixerEditorList::ExecuteListViewSearchOnAllRows(
+	const FString& SearchString, const bool bShouldRefreshAfterward)
+{
+	TArray<FString> Tokens;
+	
+	// unquoted search equivalent to a match-any-of search
+	SearchString.ParseIntoArray(Tokens, TEXT("|"), true);
+	
+	for (const TSharedPtr<FObjectMixerEditorListRow>& ChildRow : TreeViewRootObjects)
+	{
+		if (!ensure(ChildRow.IsValid()))
+		{
+			continue;
+		}
+		
+		const bool bGroupMatch = ChildRow->MatchSearchTokensToSearchTerms(Tokens);
+		
+		// If the group name matches then we pass in an empty string so all child nodes are visible.
+		// If the name doesn't match, then we need to evaluate each child.
+		ChildRow->ExecuteSearchOnChildNodes(bGroupMatch ? "" : SearchString);
+	}
+
+	if (bShouldRefreshAfterward)
+	{
+		FindVisibleObjectsAndRequestTreeRefresh();
+	}
+}
+
+bool SObjectMixerEditorList::DoesTreeViewHaveVisibleChildren() const
+{
+	if (TreeViewPtr.IsValid())
+	{
+		for (const TSharedPtr<FObjectMixerEditorListRow>& Header : TreeViewRootObjects)
+		{
+			const EVisibility HeaderVisibility = Header->GetDesiredVisibility();
+			
+			if (HeaderVisibility != EVisibility::Hidden && HeaderVisibility != EVisibility::Collapsed)
+			{
+				return true;
+			}
+		}
+	}
+	
+	return false;
+}
+
+void SObjectMixerEditorList::SetTreeViewItemExpanded(const TSharedPtr<FObjectMixerEditorListRow>& RowToExpand, const bool bNewExpansion) const
+{
+	if (TreeViewPtr.IsValid())
+	{
+		TreeViewPtr->SetItemExpansion(RowToExpand, bNewExpansion);
+	}
+}
+
+void SObjectMixerEditorList::ToggleFilterActive(const FString& FilterName)
+{
+	if (const TSharedRef<IObjectMixerEditorListFilter>* Match =
+		Algo::FindByPredicate(ShowFilters,
+		[&FilterName](TSharedRef<IObjectMixerEditorListFilter> Comparator)
+		{
+			return Comparator->GetFilterName().Equals(FilterName);
+		}))
+	{
+		const TSharedRef<IObjectMixerEditorListFilter> Filter = *Match;
+		Filter->ToggleFilterActive();
+
+		EvaluateIfRowsPassFilters();
+	}
+}
+
+void SObjectMixerEditorList::EvaluateIfRowsPassFilters(const bool bShouldRefreshAfterward)
+{
+	// Separate filters by type
+	
+	TSet<TSharedRef<IObjectMixerEditorListFilter>> MatchAnyOfFilters;
+	TSet<TSharedRef<IObjectMixerEditorListFilter>> MatchAllOfFilters;
+
+	for (const TSharedRef<IObjectMixerEditorListFilter>& Filter : ShowFilters)
+	{
+		if (Filter->GetFilterMatchType() ==
+			IObjectMixerEditorListFilter::EObjectMixerEditorListFilterMatchType::MatchAll)
+		{
+			MatchAllOfFilters.Add(Filter);
+		}
+		else
+		{
+			MatchAnyOfFilters.Add(Filter);
+		}
+	}
+	
+	for (const FObjectMixerEditorListRowPtr& Row : TreeViewRootObjects)
+	{
+		if (Row.IsValid() && Row->GetRowType() == FObjectMixerEditorListRow::SingleItem)
+		{
+			auto Projection = [&Row](const TSharedRef<IObjectMixerEditorListFilter>& Filter)
+			{
+				return Filter->GetIsFilterActive() ? Filter->DoesItemPassFilter(Row) : true;
+			};
+			
+			const bool bPassesAnyOf = MatchAnyOfFilters.Num() ? Algo::AnyOf(MatchAnyOfFilters, Projection) : true;
+			const bool bPassesAllOf = MatchAllOfFilters.Num() ? Algo::AllOf(MatchAllOfFilters, Projection) : true;
+			
+			Row->SetDoesRowPassFilters(bPassesAnyOf && bPassesAllOf);
+		}
+	}
+
+	if (bShouldRefreshAfterward)
+	{
+		FindVisibleObjectsAndRequestTreeRefresh();
+	}
+}
+
+EColumnSortMode::Type SObjectMixerEditorList::GetSortModeForColumn(FName InColumnName) const
+{
+	EColumnSortMode::Type ColumnSortMode = EColumnSortMode::None;
+
+	if (GetActiveSortingColumnName().IsEqual(InColumnName))
+	{
+		ColumnSortMode = ActiveSortingType;
+	}
+
+	return ColumnSortMode;
+}
+
+void SObjectMixerEditorList::OnSortColumnCalled(EColumnSortPriority::Type Priority, const FName& ColumnName, EColumnSortMode::Type SortMode)
+{
+	ExecuteSort(ColumnName, CycleSortMode(ColumnName));
+}
+
+EColumnSortMode::Type SObjectMixerEditorList::CycleSortMode(const FName& InColumnName)
+{
+	const EColumnSortMode::Type PreviousColumnSortMode = GetSortModeForColumn(InColumnName);
+	ActiveSortingType = PreviousColumnSortMode ==
+		EColumnSortMode::Ascending ? EColumnSortMode::Descending : EColumnSortMode::Ascending;
+
+	ActiveSortingColumnName = InColumnName;
+	return ActiveSortingType;
+}
+
+void SObjectMixerEditorList::ExecuteSort(
+	const FName& InColumnName, const EColumnSortMode::Type InColumnSortMode, const bool bShouldRefreshAfterward)
+{	
+	if (bShouldRefreshAfterward)
+	{
+		FindVisibleObjectsAndRequestTreeRefresh();
+	}
+}
+
+FListViewColumnInfo* SObjectMixerEditorList::GetColumnInfoByPropertyName(const FName& InPropertyName)
+{
+	return Algo::FindByPredicate(ListViewColumns,
+		[InPropertyName] (const FListViewColumnInfo& ColumnInfo)
+		{
+			return ColumnInfo.PropertyName.IsEqual(InPropertyName);
+		});
+}
+
+TSharedRef<SWidget> SObjectMixerEditorList::GenerateHeaderRowContextMenu() const
+{
+	FMenuBuilder MenuBuilder(false, nullptr);
+	
+	MenuBuilder.AddSearchWidget();
+
+	FName LastCategoryName = NAME_None;
+
+	for (const FListViewColumnInfo& ColumnInfo : ListViewColumns)
+	{
+		const FName& CategoryName = ColumnInfo.CategoryName;
+
+		if (!CategoryName.IsEqual(LastCategoryName))
+		{
+			LastCategoryName = CategoryName;
+			
+			MenuBuilder.EndSection();
+            MenuBuilder.BeginSection(LastCategoryName, FText::FromName(LastCategoryName));
+		}
+		
+		const FName& PropertyName = ColumnInfo.PropertyName;
+		
+		const FText Tooltip = ColumnInfo.PropertyRef ?
+			ColumnInfo.PropertyRef->GetToolTipText() : ColumnInfo.PropertyDisplayText;
+
+		const bool bCanSelectColumn = ColumnInfo.PropertyType != EListViewColumnType::BuiltIn;
+
+		const FName Hook = ColumnInfo.PropertyType == EListViewColumnType::BuiltIn ? "Builtin" : "GeneratedProperties";
+		
+		MenuBuilder.AddMenuEntry(
+			ColumnInfo.PropertyDisplayText,
+			Tooltip,
+			FSlateIcon(),
+			FUIAction(
+				FExecuteAction::CreateLambda([this, PropertyName]()
+				{
+					check(HeaderRow);
+					
+					HeaderRow->SetShowGeneratedColumn(PropertyName, !HeaderRow->IsColumnVisible(PropertyName));
+				}),
+				FCanExecuteAction::CreateLambda([bCanSelectColumn](){return bCanSelectColumn;}),
+				FIsActionChecked::CreateLambda([this, PropertyName]()
+				{
+					check(HeaderRow);
+					
+					return HeaderRow->IsColumnVisible(PropertyName);
+				})
+			),
+			Hook,
+			EUserInterfaceActionType::Check
+		);
+	}
+
+	return  MenuBuilder.MakeWidget();
+}
+
+bool SObjectMixerEditorList::AddUniquePropertyColumnsToHeaderRow(
+	FProperty* Property,
+	const bool bForceIncludeProperty,
+	const TArray<FName>& PropertySkipList)
+{
+	if (!ensureAlwaysMsgf(Property, TEXT("%hs: Invalid property passed in. Please ensure only valid properties are passed to this function."), __FUNCTION__))
+	{
+		return false;
+	}
+
+	bool bShouldIncludeProperty = bForceIncludeProperty;
+
+	if (!bShouldIncludeProperty)
+	{
+		const bool bIsPropertyBlueprintEditable = (Property->GetPropertyFlags() & CPF_Edit) != 0;
+
+		// We don't have a proper way to display these yet
+		const bool bDoesPropertyHaveSupportedClass =
+			!Property->IsA(FMapProperty::StaticClass()) &&
+			!Property->IsA(FArrayProperty::StaticClass()) &&
+			!Property->IsA(FSetProperty::StaticClass()) &&
+			!Property->IsA(FStructProperty::StaticClass());
+
+		bShouldIncludeProperty = bIsPropertyBlueprintEditable && bDoesPropertyHaveSupportedClass;
+	}
+
+	if (bShouldIncludeProperty)
+	{
+		const bool bIsPropertyExplicitlySkipped =
+		   PropertySkipList.Num() && PropertySkipList.Contains(Property->GetFName());
+
+		bShouldIncludeProperty = !bIsPropertyExplicitlySkipped;
+	}
+	
+	if (bShouldIncludeProperty)
+	{
+		const FName PropertyName = Property->GetFName();
+	
+		// Ensure no duplicate properties
+		if (!Algo::FindByPredicate(ListViewColumns,
+				[&PropertyName] (const FListViewColumnInfo& ListViewColumn)
+				{
+					return ListViewColumn.PropertyName.IsEqual(PropertyName);
+				})
+			)
+		{
+			ListViewColumns.Add(
+				{
+					Property, PropertyName,
+					Property->GetDisplayNameText(),
+					EListViewColumnType::PropertyGenerated,
+					"Generated Properties",
+					true, true, false
+				}
+			);
+
+			return true;
+		}
+	}
+
+	return false;
+}
+
+void SObjectMixerEditorList::AddBuiltinColumnsToHeaderRow()
+{
+	ListViewColumns.Insert(
+		{
+			nullptr, ItemNameColumnName,
+			LOCTEXT("ItemNameHeaderText", "Name"),
+			EListViewColumnType::BuiltIn, "Built-In",
+			true, true,
+			false, 1.0f, 1.7f
+		}, 0
+	);
+	
+	ListViewColumns.Insert(
+		{
+			nullptr, EditorVisibilitySoloColumnName,
+			LOCTEXT("EditorVisibilitySoloColumnNameHeaderText", "Solo"),
+			EListViewColumnType::BuiltIn, "Built-In",
+			true, false,
+			true, 25.0f
+		}, 0
+	);
+
+	ListViewColumns.Insert(
+		{
+			nullptr, EditorVisibilityColumnName,
+			LOCTEXT("EditorVisibilityColumnNameHeaderText", "Visibility"),
+			EListViewColumnType::BuiltIn, "Built-In",
+			true, false,
+			true, 25.0f
+		}, 0
+	);
+}
+
+TSharedPtr<SHeaderRow> SObjectMixerEditorList::GenerateHeaderRow()
+{
+	check(ListModelPtr.IsValid());
+	check(HeaderRow);
+
+	TMap<FName, bool> LastVisibleColumns;
+	for (SHeaderRow::FColumn Column : HeaderRow->GetColumns())
+	{
+		LastVisibleColumns.Add(Column.ColumnId, Column.bIsVisible);
+	}
+	
+	HeaderRow->ClearColumns();
+	ListViewColumns.Empty(ListViewColumns.Num());
+
+	// Property Columns
+	const TObjectPtr<UObjectMixerObjectFilter> SelectedFilter = ListModelPtr.Pin()->GetObjectFilter();
+	if (!SelectedFilter)
+	{
+		UE_LOG(LogObjectMixerEditor, Display, TEXT("%hs: No classes defined in UObjectMixerObjectFilter class."), __FUNCTION__);
+		return nullptr;
+	}
+
+	const EObjectMixerPropertyInheritanceInclusionOptions Options =
+		SelectedFilter->GetObjectMixerPropertyInheritanceInclusionOptions();
+	TArray<UClass*> SpecifiedClasses =
+		SelectedFilter->GetParentAndChildClassesFromSpecifiedClasses(SelectedFilter->GetObjectClassesToFilter(), Options);
+	
+	const TArray<FName> PropertySkipList = SelectedFilter->GetColumnsFilter();
+	const TArray<FName> ForceAddedPropertyList = SelectedFilter->GetForceAddedColumns();
+
+	const bool bShouldIncludeUnsupportedProperties = SelectedFilter->ShouldIncludeUnsupportedProperties();
+	
+	for (const UClass* Class : SpecifiedClasses)
+	{
+		for (TFieldIterator<FProperty> FieldIterator(Class, EFieldIterationFlags::None); FieldIterator; ++FieldIterator)
+		{
+			if (FProperty* Property = *FieldIterator)
+			{
+				AddUniquePropertyColumnsToHeaderRow(Property, bShouldIncludeUnsupportedProperties, PropertySkipList);
+			}
+		}
+
+		// Check Force Added Columns
+		for (const FName& PropertyName : ForceAddedPropertyList)
+		{
+			if (FProperty* Property = FindFProperty<FProperty>(Class, PropertyName))
+			{
+				AddUniquePropertyColumnsToHeaderRow(Property, true);
+			}
+		}
+	}
+
+	// Alphabetical sort by Property Name
+	ListViewColumns.StableSort([](const FListViewColumnInfo& A, const FListViewColumnInfo& B)
+	{
+		return A.PropertyDisplayText.ToString() < B.PropertyDisplayText.ToString();
+	});
+
+	// Alphabetical sort by Category Name
+	ListViewColumns.StableSort([](const FListViewColumnInfo& A, const FListViewColumnInfo& B)
+	{
+		return A.CategoryName.LexicalLess(B.CategoryName);
+	});
+
+	// Add Built-in Columns to beginning
+	AddBuiltinColumnsToHeaderRow();
+
+	// Actually add columns to Header
+	const FText ClickToSortTooltip = LOCTEXT("ClickToSort","Click to sort");
+	
+	const TArray<FName> ColumnsToShowByDefault = SelectedFilter->GetColumnsToShowByDefault();
+
+	const TSharedRef<SWidget> HeaderMenuContent = GenerateHeaderRowContextMenu();
+
+	for (const FListViewColumnInfo& ColumnInfo : ListViewColumns)
+	{
+		const FText Tooltip = ColumnInfo.PropertyRef ? ColumnInfo.PropertyRef->GetToolTipText() :
+			ColumnInfo.bCanBeSorted ? ClickToSortTooltip : ColumnInfo.PropertyDisplayText;
+		
+		SHeaderRow::FColumn::FArguments Column =
+			SHeaderRow::Column(ColumnInfo.PropertyName)
+			.DefaultLabel(ColumnInfo.PropertyDisplayText)
+			.ToolTipText(Tooltip)
+			.HAlignHeader(EHorizontalAlignment::HAlign_Left)
+		;
+
+		if (ColumnInfo.bUseFixedWidth)
+		{
+			Column.FixedWidth(ColumnInfo.FixedWidth);
+		}
+		else
+		{
+			Column.FillWidth(ColumnInfo.FillWidth);
+		}
+
+		if (ColumnInfo.bCanBeSorted)
+		{
+			Column.SortMode_Raw(this, &SObjectMixerEditorList::GetSortModeForColumn, ColumnInfo.PropertyName);
+			Column.OnSort_Raw(this, &SObjectMixerEditorList::OnSortColumnCalled);
+		}
+
+		if (ColumnInfo.PropertyType == EListViewColumnType::BuiltIn)
+		{
+			Column.ShouldGenerateWidget(true);
+		}
+
+		if (ColumnInfo.PropertyName.IsEqual(EditorVisibilityColumnName))
+		{
+			Column.HeaderContent()
+			[
+				SNew(SBox)
+				.HAlign(HAlign_Left)
+				.VAlign(VAlign_Center)
+				.Padding(0.f)
+				[
+					SNew(SImage)
+					.ColorAndOpacity(FSlateColor::UseForeground())
+					.Image(FAppStyle::Get().GetBrush("Level.VisibleIcon16x"))
+				]
+			];
+		}
+		else if (ColumnInfo.PropertyName.IsEqual(EditorVisibilitySoloColumnName))
+		{
+			Column.HeaderContent()
+			[
+				SNew(SBox)
+				.HAlign(HAlign_Left)
+				.VAlign(VAlign_Center)
+				.Padding(0.f)
+				[
+					SNew(SImage)
+					.ColorAndOpacity(FSlateColor::UseForeground())
+					.Image(FAppStyle::Get().GetBrush("MediaAsset.AssetActions.Solo.Small"))
+				]
+			];
+		}
+		else // Add Column Selection Menu widget to all other columns
+		{
+			Column.MenuContent()
+			[
+				HeaderMenuContent
+			];
+		}
+		
+		HeaderRow->AddColumn(Column);
+		bool bShouldShowColumn = ColumnsToShowByDefault.Contains(ColumnInfo.PropertyName);
+
+		if (const bool* Match = LastVisibleColumns.Find(ColumnInfo.PropertyName))
+		{
+			bShouldShowColumn = *Match;
+		}
+		
+		HeaderRow->SetShowGeneratedColumn(ColumnInfo.PropertyName, bShouldShowColumn);
+	}
+
+	return HeaderRow;
+}
+
+void SObjectMixerEditorList::SetupFilters()
+{
+	
+}
+
+TSharedRef<SWidget> SObjectMixerEditorList::BuildShowOptionsMenu()
+{
+	FMenuBuilder ShowOptionsMenuBuilder = FMenuBuilder(true, nullptr);
+
+	ShowOptionsMenuBuilder.AddMenuEntry(
+		FText::FromString("Open Generic Object Mixer Instance"),
+		FText::FromString("Open Generic Object Mixer Instance"),
+		FSlateIcon(),
+		FUIAction(FExecuteAction::CreateLambda([]()
+		{
+			FGlobalTabmanager::Get()->TryInvokeTab(FObjectMixerEditorModule::ObjectMixerToolkitPanelTabId);
+		})));
+
+	ShowOptionsMenuBuilder.AddMenuEntry(FText::FromString("Refresh List"), FText::FromString("Refresh"), FSlateIcon(),
+		FUIAction(FExecuteAction::CreateRaw(this, &SObjectMixerEditorList::GenerateTreeView)));
+
+	ShowOptionsMenuBuilder.BeginSection("", LOCTEXT("ShowOptions_ShowSectionHeading", "Show"));
+	{
+		// Add show filters
+		auto AddFiltersLambda = [this, &ShowOptionsMenuBuilder](const TSharedRef<IObjectMixerEditorListFilter>& InFilter)
+		{
+			const FString& FilterName = InFilter->GetFilterName();
+			
+			ShowOptionsMenuBuilder.AddMenuEntry(
+			   InFilter->GetFilterButtonLabel(),
+			   InFilter->GetFilterButtonToolTip(),
+			   FSlateIcon(),
+			   FUIAction(
+				   FExecuteAction::CreateLambda(
+				   	[this, FilterName]()
+					   {
+						   ToggleFilterActive(FilterName);
+					   }
+					),
+				   FCanExecuteAction(),
+				   FIsActionChecked::CreateSP( InFilter, &IObjectMixerEditorListFilter::GetIsFilterActive )
+			   ),
+			   NAME_None,
+			   EUserInterfaceActionType::ToggleButton
+		   );
+		};
+
+		for (const TSharedRef<IObjectMixerEditorListFilter>& Filter : ShowFilters)
+		{
+			AddFiltersLambda(Filter);
+		}
+	}
+	ShowOptionsMenuBuilder.EndSection();
+
+	return ShowOptionsMenuBuilder.MakeWidget();
+}
+
+void SObjectMixerEditorList::FlushMemory(const bool bShouldKeepMemoryAllocated)
+{
+	if (bShouldKeepMemoryAllocated)
+	{
+		TreeViewRootObjects.Reset();
+		VisibleTreeViewObjects.Reset();
+	}
+	else
+	{
+		TreeViewRootObjects.Empty();
+		VisibleTreeViewObjects.Empty();
+	}
+}
+
+void SObjectMixerEditorList::SetAllGroupsCollapsed()
+{
+	if (TreeViewPtr.IsValid())
+	{
+		for (const FObjectMixerEditorListRowPtr& RootRow : TreeViewRootObjects)
+		{
+			if (!RootRow.IsValid())
+			{
+				continue;
+			}
+			
+			TreeViewPtr->SetItemExpansion(RootRow, false);
+			RootRow->SetIsTreeViewItemExpanded(false);
+		}
+	}
+}
+
+void SObjectMixerEditorList::OnListViewSearchTextChanged(const FText& Text)
+{
+	ExecuteListViewSearchOnAllRows(Text.ToString(), true);
+}
+
+void SObjectMixerEditorList::GenerateTreeView()
+{
+	check(ListModelPtr.IsValid());
+	
+	if (!ensure(TreeViewPtr.IsValid()))
+	{
+		return;
+	}
+	
+	FlushMemory(true);
+	
+	GenerateHeaderRow();
+
+	TArray<UClass*> AcceptableClasses = ListModelPtr.Pin()->GetObjectClasses();
+
+	check(GEditor);
+	const UWorld* EditorWorld = GEditor->GetEditorWorldContext().World();
+	
+	for (TObjectIterator<UObject> ObjectIterator; ObjectIterator; ++ObjectIterator)
+	{
+		if (UObject* Object = *ObjectIterator; Object->GetWorld() == EditorWorld)
+		{
+			bool bIsAcceptableClass = false;
+
+			for (UClass* Class : AcceptableClasses)
+			{
+				if (Object->IsA(Class))
+				{
+					bIsAcceptableClass = true;
+					break;
+				}
+			}
+
+			if (bIsAcceptableClass)
+			{
+				TreeViewRootObjects.Add(
+					MakeShared<FObjectMixerEditorListRow>(
+						Object, FObjectMixerEditorListRow::SingleItem, SharedThis(this), nullptr)
+				);
+			}
+		}
+	}
+
+	RefreshList();
+}
+
+void SObjectMixerEditorList::FindVisibleTreeViewObjects()
+{
+	VisibleTreeViewObjects.Empty();
+
+	for (const TSharedPtr<FObjectMixerEditorListRow>& Row : TreeViewRootObjects)
+	{
+		if (Row->ShouldBeVisible())
+		{
+			VisibleTreeViewObjects.Add(Row);
+		}
+	}
+}
+
+void SObjectMixerEditorList::FindVisibleObjectsAndRequestTreeRefresh()
+{
+	FindVisibleTreeViewObjects();
+	TreeViewPtr->RequestTreeRefresh();
+}
+
+void SObjectMixerEditorList::OnGetRowChildren(FObjectMixerEditorListRowPtr Row, TArray<FObjectMixerEditorListRowPtr>& OutChildren) const
+{
+	if (Row.IsValid())
+	{
+		OutChildren = Row->GetChildRows();
+
+		if (Row->GetShouldExpandAllChildren())
+		{
+			SetChildExpansionRecursively(Row, true);
+			Row->SetShouldExpandAllChildren(false);
+		}
+	}
+}
+
+void SObjectMixerEditorList::OnRowChildExpansionChange(FObjectMixerEditorListRowPtr Row, const bool bIsExpanded, const bool bIsRecursive) const
+{
+	if (Row.IsValid())
+	{
+		if (bIsRecursive)
+		{
+			if (bIsExpanded)
+			{
+				if (Row->GetRowType() == FObjectMixerEditorListRow::Group)
+				{
+					Row->SetShouldExpandAllChildren(true);
+				}
+			}
+			else
+			{
+				SetChildExpansionRecursively(Row, bIsExpanded);
+			}
+		}
+		
+		TreeViewPtr->SetItemExpansion(Row, bIsExpanded);
+		Row->SetIsTreeViewItemExpanded(bIsExpanded);
+	}
+}
+
+void SObjectMixerEditorList::SetChildExpansionRecursively(const FObjectMixerEditorListRowPtr& InRow, const bool bNewIsExpanded) const
+{
+	if (InRow.IsValid())
+	{
+		for (const FObjectMixerEditorListRowPtr& Child : InRow->GetChildRows())
+		{
+			TreeViewPtr->SetItemExpansion(Child, bNewIsExpanded);
+			Child->SetIsTreeViewItemExpanded(bNewIsExpanded);
+
+			SetChildExpansionRecursively(Child, bNewIsExpanded);
+		}
+	}
+}
+
+#undef LOCTEXT_NAMESPACE
diff --git a/Engine/Source/Editor/ObjectMixer/Private/Views/List/SObjectMixerEditorListRow.cpp b/Engine/Source/Editor/ObjectMixer/Private/Views/List/SObjectMixerEditorListRow.cpp
new file mode 100644
--- /dev/null
+++ b/Engine/Source/Editor/ObjectMixer/Private/Views/List/SObjectMixerEditorListRow.cpp
@@ -0,0 +1,586 @@
+// Copyright Epic Games, Inc. All Rights Reserved.
+
+#include "Views/List/SObjectMixerEditorListRow.h"
+
+#include "ObjectMixerEditorStyle.h"
+#include "Views/List/ObjectMixerEditorList.h"
+#include "Views/List/SObjectMixerEditorList.h"
+
+#include "Customizations/ColorStructCustomization.h"
+#include "Input/DragAndDrop.h"
+#include "ISinglePropertyView.h"
+#include "IStructureDetailsView.h"
+#include "Modules/ModuleManager.h"
+#include "PropertyEditorDelegates.h"
+#include "PropertyEditorModule.h"
+#include "PropertyHandle.h"
+#include "ScopedTransaction.h"
+#include "Styling/AppStyle.h"
+#include "Styling/StyleColors.h"
+#include "Widgets/Input/SCheckBox.h"
+#include "Widgets/Layout/SBox.h"
+#include "Widgets/SBoxPanel.h"
+#include "Widgets/SOverlay.h"
+#include "Widgets/Layout/SScaleBox.h"
+#include "Widgets/Text/STextBlock.h"
+
+#define LOCTEXT_NAMESPACE "ObjectMixerEditor"
+
+// const FText InsertFormatText = LOCTEXT("InsertAboveFormatText", "Insert {0} {1} {2}");
+// const FText AboveText = LOCTEXT("AboveListItem", "above");
+// const FText BelowText = LOCTEXT("BelowListItem", "below");
+// const FText MultiDragFormatText = LOCTEXT("MultiDragFormatText", "{0} Items");
+
+// class FObjectMixerListRowDragDropOp : public FDecoratedDragDropOp
+// {
+// public:
+// 	DRAG_DROP_OPERATOR_TYPE(FObjectMixerListRowDragDropOp, FDecoratedDragDropOp)
+//
+// 	/** The item being dragged and dropped */
+// 	TArray<FObjectMixerEditorListRowPtr> DraggedItems;
+//
+// 	/** Constructs a new drag/drop operation */
+// 	static TSharedRef<FObjectMixerListRowDragDropOp> New(const TArray<FObjectMixerEditorListRowPtr>& InItems)
+// 	{
+// 		check(InItems.Num() > 0);
+//
+// 		TSharedRef<FObjectMixerListRowDragDropOp> Operation = MakeShareable(
+// 			new FObjectMixerListRowDragDropOp());
+//
+// 		Operation->DraggedItems = InItems;
+//
+// 		Operation->DefaultHoverIcon = FAppStyle::Get().GetBrush("Graph.ConnectorFeedback.Error");
+//
+// 		// Set the display text and the transaction name based on whether we're dragging a single or multiple widgets
+// 		if (InItems.Num() == 1)
+// 		{
+// 			//Operation->DefaultHoverText = FText::FromString(InItems[0]->GetCommandInfo().Pin()->Command);
+// 		}
+// 		else
+// 		{
+// 			Operation->DefaultHoverText =
+// 				FText::Format(
+// 					SObjectMixerEditorListRow::MultiDragFormatText,
+// 					FText::AsNumber(Operation->DraggedItems.Num())
+// 				);
+// 		}
+//
+// 		Operation->Construct();
+//
+// 		return Operation;
+// 	}
+// };
+
+void SObjectMixerEditorListRow::Construct(
+	const FArguments& InArgs, const TSharedRef<STableViewBase>& InOwnerTable,
+	const TWeakPtr<FObjectMixerEditorListRow> InRow)
+{
+	check(InRow.IsValid());
+
+	Item = InRow;
+	const FObjectMixerEditorListRowPtr PinnedItem = Item.Pin();
+
+	// Set up flash animation
+	FlashAnimation = FCurveSequence(0.f, FlashAnimationDuration, ECurveEaseFunction::QuadInOut);
+
+	SMultiColumnTableRow<FObjectMixerEditorListRowPtr>::Construct(
+		FSuperRowType::FArguments()
+		.Padding(1.0f)
+	// 	.OnCanAcceptDrop(this, &SObjectMixerEditorListRow::HandleCanAcceptDrop)
+	// 	.OnAcceptDrop(this, &SObjectMixerEditorListRow::HandleAcceptDrop)
+	// 	.OnDragDetected(this, &SObjectMixerEditorListRow::HandleDragDetected)
+	// 	.OnDragLeave(this, &SObjectMixerEditorListRow::HandleDragLeave)
+	 	, InOwnerTable
+	);
+
+	const FName VisibleHoveredBrushName = TEXT("Level.VisibleHighlightIcon16x");
+	const FName VisibleNotHoveredBrushName = TEXT("Level.VisibleIcon16x");
+	const FName NotVisibleHoveredBrushName = TEXT("Level.NotVisibleHighlightIcon16x");
+	const FName NotVisibleNotHoveredBrushName = TEXT("Level.NotVisibleIcon16x");
+
+	VisibleHoveredBrush = FAppStyle::Get().GetBrush(VisibleHoveredBrushName);
+	VisibleNotHoveredBrush = FAppStyle::Get().GetBrush(VisibleNotHoveredBrushName);
+	NotVisibleHoveredBrush = FAppStyle::Get().GetBrush(NotVisibleHoveredBrushName);
+	NotVisibleNotHoveredBrush = FAppStyle::Get().GetBrush(NotVisibleNotHoveredBrushName);
+
+	if (PinnedItem->GetShouldFlashOnScrollIntoView())
+	{
+		FlashRow();
+
+		PinnedItem->SetShouldFlashOnScrollIntoView(false);
+	}
+}
+
+TSharedRef<SWidget> SObjectMixerEditorListRow::GenerateWidgetForColumn(const FName& InColumnName)
+{
+	const FObjectMixerEditorListRowPtr PinnedItem = Item.Pin();
+
+	const TSharedPtr<SWidget> CellWidget = GenerateCells(InColumnName, PinnedItem);
+
+	const TSharedRef<SImage> FlashImage = SNew(SImage)
+										.Image(new FSlateColorBrush(FStyleColors::White))
+										.Visibility_Raw(
+											this, &SObjectMixerEditorListRow::GetFlashImageVisibility)
+										.ColorAndOpacity_Raw(
+											this, &SObjectMixerEditorListRow::GetFlashImageColorAndOpacity);
+
+	FlashImages.Add(FlashImage);
+
+	return SNew(SOverlay)
+			.Visibility(EVisibility::SelfHitTestInvisible)
+
+			+ SOverlay::Slot()
+			[
+				FlashImage
+			]
+
+			+ SOverlay::Slot()
+			[
+				SNew(SBorder)
+					.HAlign(HAlign_Fill)
+					.VAlign(VAlign_Center)
+					.BorderImage(GetBorderImage(PinnedItem->GetRowType()))
+				[
+					CellWidget.ToSharedRef()
+				]
+			];
+}
+
+void SObjectMixerEditorListRow::OnMouseEnter(const FGeometry& MyGeometry, const FPointerEvent& MouseEvent)
+{
+	bIsHovered = true;
+
+	SMultiColumnTableRow<FObjectMixerEditorListRowPtr>::OnMouseEnter(MyGeometry, MouseEvent);
+}
+
+void SObjectMixerEditorListRow::OnMouseLeave(const FPointerEvent& MouseEvent)
+{
+	bIsHovered = false;
+
+	SMultiColumnTableRow<FObjectMixerEditorListRowPtr>::OnMouseLeave(MouseEvent);
+}
+
+SObjectMixerEditorListRow::~SObjectMixerEditorListRow()
+{
+	Item.Reset();
+	HoverToolTip.Reset();
+
+	FlashImages.Empty();
+
+	ValueChildInputWidget.Reset();
+
+	HoverableWidgetsPtr.Reset();
+}
+
+// FReply SObjectMixerEditorListRow::HandleDragDetected(const FGeometry& MyGeometry, const FPointerEvent& MouseEvent)
+// {
+// 	TArray<FObjectMixerEditorListRowPtr> DraggedItems = Item.Pin()->GetSelectedTreeViewItems();
+// 	TSharedRef<FObjectMixerListRowDragDropOp> Operation =
+// 		FObjectMixerListRowDragDropOp::New(DraggedItems);
+//
+// 	return FReply::Handled().BeginDragDrop(Operation);
+// }
+//
+// void SObjectMixerEditorListRow::HandleDragLeave(const FDragDropEvent& DragDropEvent)
+// {
+// 	if (TSharedPtr<FObjectMixerListRowDragDropOp> Operation =
+// 		DragDropEvent.GetOperationAs<FObjectMixerListRowDragDropOp>())
+// 	{
+// 		Operation->ResetToDefaultToolTip();
+// 	}
+// }
+
+// TOptional<EItemDropZone> SObjectMixerEditorListRow::HandleCanAcceptDrop(const FDragDropEvent& DragDropEvent,
+//                                                                              EItemDropZone DropZone,
+//                                                                              FObjectMixerEditorListRowPtr
+//                                                                              TargetItem)
+// {
+// 	TSharedPtr<FObjectMixerListRowDragDropOp> Operation =
+// 		DragDropEvent.GetOperationAs<FObjectMixerListRowDragDropOp>();
+//
+// 	if (!Operation.IsValid())
+// 	{
+// 		return TOptional<EItemDropZone>();
+// 	}
+//
+// 	Operation->SetToolTip(
+// 		LOCTEXT("SortByCustomOrderDrgDropWarning", "Sort by custom order (\"#\") to drag & drop"),
+// 		FAppStyle::Get().GetBrush("Graph.ConnectorFeedback.Error")
+// 	);
+//
+// 	const bool bIsDropDenied = false;
+//
+// 	FString BoolAsString = bIsDropDenied ? "true" : "false";
+//
+// 	if (bIsDropDenied)
+// 	{
+// 		Operation->ResetToDefaultToolTip();
+//
+// 		return TOptional<EItemDropZone>();
+// 	}
+//
+// 	FText ItemNameText = FText::FromString("Light Label");
+//
+// 	if (Operation->DraggedItems.Num() > 1)
+// 	{
+// 		ItemNameText = FText::Format(MultiDragFormatText, FText::AsNumber(Operation->DraggedItems.Num()));
+// 	}
+//
+// 	const FText DropPermittedText =
+// 		FText::Format(InsertFormatText,
+// 		              ItemNameText,
+// 		              DropZone == EItemDropZone::BelowItem ? BelowText : AboveText,
+// 		              ItemNameText
+// 		);
+//
+// 	Operation->SetToolTip(
+// 		DropPermittedText,
+// 		FAppStyle::Get().GetBrush("Graph.ConnectorFeedback.OK")
+// 	);
+//
+// 	// We have no behaviour yet for dropping one item onto another, so we'll treat it like we dropped it above
+// 	return DropZone == EItemDropZone::OntoItem ? EItemDropZone::AboveItem : DropZone;
+// }
+
+// FReply SObjectMixerEditorListRow::HandleAcceptDrop(const FDragDropEvent& DragDropEvent, EItemDropZone DropZone,
+//                                                         FObjectMixerEditorListRowPtr TargetItem)
+// {
+// 	TSharedPtr<FObjectMixerListRowDragDropOp> Operation =
+// 		DragDropEvent.GetOperationAs<FObjectMixerListRowDragDropOp>();
+//
+// 	if (!Operation.IsValid())
+// 	{
+// 		return FReply::Unhandled();
+// 	}
+//
+// 	const TSharedPtr<SObjectMixerEditorList> ListView = Item.Pin()->GetListViewPtr().Pin();
+//
+// 	TArray<FObjectMixerEditorListRowPtr> DraggedItems = Operation->DraggedItems;
+//
+// 	TArray<FObjectMixerEditorListRowPtr> AllTreeItemsCopy = ListView->GetTreeViewItems();
+//
+// 	for (FObjectMixerEditorListRowPtr DraggedItem : DraggedItems)
+// 	{
+// 		if (!DraggedItem.IsValid() || !AllTreeItemsCopy.Contains(DraggedItem))
+// 		{
+// 			continue;
+// 		}
+//
+// 		AllTreeItemsCopy.Remove(DraggedItem);
+// 	}
+//
+// 	const int32 TargetIndex = AllTreeItemsCopy.IndexOfByKey(TargetItem);
+//
+// 	if (TargetIndex > -1)
+// 	{
+// 		for (int32 ItemIndex = DraggedItems.Num() - 1; ItemIndex >= 0; ItemIndex--)
+// 		{
+// 			FObjectMixerEditorListRowPtr DraggedItem = DraggedItems[ItemIndex];
+//
+// 			if (!DraggedItem.IsValid() || AllTreeItemsCopy.Contains(DraggedItem))
+// 			{
+// 				continue;
+// 			}
+//
+// 			AllTreeItemsCopy.Insert(DraggedItem, DropZone == EItemDropZone::AboveItem ? TargetIndex : TargetIndex + 1);
+// 		}
+//
+// 		ListView->SetTreeViewItems(AllTreeItemsCopy);
+// 	}
+//
+// 	return FReply::Handled();
+// }
+
+void SObjectMixerEditorListRow::FlashRow()
+{
+	FlashAnimation.Play(this->AsShared());
+}
+
+bool SObjectMixerEditorListRow::IsVisible() const
+{
+	if (const TSharedPtr<FObjectMixerEditorListRow> PinnedItem = Item.Pin())
+	{
+		if (const TSharedPtr<SObjectMixerEditorList> PinnedListView = PinnedItem->GetListViewPtr().Pin())
+		{
+			if (const TSharedPtr<FObjectMixerEditorList> PinnedListModel = PinnedListView->GetListModelPtr().Pin())
+			{
+				if (const UObjectMixerObjectFilter* Filter = PinnedListModel->GetObjectFilter())
+				{
+					return Filter->GetRowEditorVisibility(PinnedItem->GetObject());
+				}
+			}
+		}
+	}
+
+	return false;
+}
+
+FSlateColor SObjectMixerEditorListRow::GetVisibilityIconForegroundColor() const
+{
+	check(Item.IsValid());
+	
+	const bool bIsSelected = Item.Pin()->GetIsSelected();
+
+	// make the foreground brush transparent if it is not selected and it is visible
+	if (IsVisible() && !bIsHovered && !bIsSelected)
+	{
+		return FLinearColor::Transparent;
+	}
+	else if (bIsHovered && !bIsSelected)
+	{
+		return FAppStyle::Get().GetSlateColor("Colors.ForegroundHover");
+	}
+
+	return FSlateColor::UseForeground();
+}
+
+FSlateColor SObjectMixerEditorListRow::GetSoloIconForegroundColor() const
+{
+	check(Item.IsValid());
+
+	if (Item.Pin()->IsThisRowSolo())
+	{
+		return bIsHovered ? FAppStyle::Get().GetSlateColor("Colors.ForegroundHover") : FSlateColor::UseForeground();
+	}
+
+	return FLinearColor::Transparent;
+}
+
+const FSlateBrush* SObjectMixerEditorListRow::GetVisibilityBrush() const
+{
+	if (IsVisible())
+	{
+		return bIsHovered ? VisibleHoveredBrush : VisibleNotHoveredBrush;
+	}
+	
+	return bIsHovered ? NotVisibleHoveredBrush : NotVisibleNotHoveredBrush;
+}
+
+EVisibility SObjectMixerEditorListRow::GetFlashImageVisibility() const
+{
+	return FlashAnimation.IsPlaying() ? EVisibility::SelfHitTestInvisible : EVisibility::Hidden;
+}
+
+FSlateColor SObjectMixerEditorListRow::GetFlashImageColorAndOpacity() const
+{
+	if (FlashAnimation.IsPlaying())
+	{
+		// This equation modulates the alpha into a parabolic curve 
+		const float Progress = FMath::Abs(FMath::Abs((FlashAnimation.GetLerp() - 0.5f) * 2) - 1);
+		return FLinearColor::LerpUsingHSV(FLinearColor::Transparent, FlashColor, Progress);
+	}
+
+	return FLinearColor::Transparent;
+}
+
+const FSlateBrush* SObjectMixerEditorListRow::GetBorderImage(
+	const FObjectMixerEditorListRow::EObjectMixerEditorListRowType InRowType)
+{
+	switch (InRowType)
+	{
+	case FObjectMixerEditorListRow::Group:
+		return FObjectMixerEditorStyle::Get().GetBrush("ObjectMixerEditor.CommandGroupBorder");
+
+	default:
+		return FObjectMixerEditorStyle::Get().GetBrush("ObjectMixerEditor.DefaultBorder");
+	}
+}
+
+TSharedRef<SWidget> SObjectMixerEditorListRow::GenerateCells(
+	const FName& InColumnName, const TSharedPtr<FObjectMixerEditorListRow> PinnedItem)
+{
+	check(PinnedItem.IsValid());
+	
+	if (InColumnName.IsEqual(SObjectMixerEditorList::ItemNameColumnName))
+	{
+		FText RowName;
+
+		if (const UObjectMixerObjectFilter* Filter = PinnedItem->GetObjectFilter())
+		{
+			RowName = Filter->GetRowDisplayName(PinnedItem->GetObject());
+		}
+
+		TSharedRef<SHorizontalBox> HBox = SNew(SHorizontalBox);
+		
+		if (const FSlateBrush* Icon = PinnedItem->GetObjectIconBrush())
+		{
+			HBox->AddSlot()
+			.AutoWidth()
+			[
+				SNew(SImage).Image(Icon)
+			];
+		}
+
+		HBox->AddSlot()
+		.Padding(FMargin(10.0, 0, 0, 0))
+		[
+			SNew(STextBlock)
+			.Visibility(EVisibility::Visible)
+			.Justification(ETextJustify::Left)
+			.Text(RowName)
+			.ToolTip(HoverToolTip)
+		];
+		
+		return SNew(SBox)
+				.Visibility(EVisibility::SelfHitTestInvisible)
+				.HAlign(HAlign_Left)
+				.VAlign(VAlign_Center)
+				.Padding(FMargin(TextBlockLeftPadding, 0.f, 0.f, 0.f))
+				[
+					HBox
+				];
+	}
+
+	if (InColumnName.IsEqual(SObjectMixerEditorList::EditorVisibilityColumnName))
+	{
+		return SNew(SBox)
+				.HAlign(HAlign_Left)
+				.VAlign(VAlign_Center)
+				.Visibility(EVisibility::SelfHitTestInvisible)
+				.Padding(0.f)
+				[
+					SNew(SImage)
+					.ColorAndOpacity(this, &SObjectMixerEditorListRow::GetVisibilityIconForegroundColor)
+					.Image_Raw(this, &SObjectMixerEditorListRow::GetVisibilityBrush)
+					.OnMouseButtonDown_Lambda(
+						[PinnedItem] (const FGeometry& MyGeometry, const FPointerEvent& Event)
+						{
+							check (PinnedItem);
+
+							FScopedTransaction Transaction( LOCTEXT("VisibilityChanged", "Object Mixer - Visibility Changed") );
+
+							const bool bIsVisible = PinnedItem->GetObjectVisibility();
+							
+							if (const TSharedPtr<SObjectMixerEditorList> PinnedListView = PinnedItem->GetListViewPtr().Pin();
+								PinnedItem->GetIsSelected() && PinnedListView->GetSelectedTreeViewItemCount() > 0)
+							{
+								PinnedListView->SetSelectedTreeViewItemActorsEditorVisible(!bIsVisible);
+
+								return FReply::Handled();
+							}
+							
+							PinnedItem->SetObjectVisibility(!bIsVisible);
+
+							return FReply::Handled();
+						}
+					)
+				]
+			;
+	}
+
+	if (InColumnName.IsEqual(SObjectMixerEditorList::EditorVisibilitySoloColumnName))
+	{
+		return SNew(SBox)
+				.HAlign(HAlign_Left)
+				.VAlign(VAlign_Center)
+				.Visibility(EVisibility::SelfHitTestInvisible)
+				.Padding(0.f)
+				[
+					SNew(SImage)
+					.ColorAndOpacity(this, &SObjectMixerEditorListRow::GetSoloIconForegroundColor)
+					.Image(FAppStyle::Get().GetBrush("MediaAsset.AssetActions.Solo.Small"))
+					.OnMouseButtonDown_Lambda(
+						[PinnedItem] (const FGeometry& MyGeometry, const FPointerEvent& Event)
+						{
+							check (PinnedItem);
+							
+							if (const TSharedPtr<SObjectMixerEditorList> PinnedListView = PinnedItem->GetListViewPtr().Pin();
+								PinnedListView->GetTreeViewItemCount() > 0)
+							{
+								const bool bIsRowSolo = PinnedItem->IsThisRowSolo();
+								for (TSharedPtr<FObjectMixerEditorListRow> Item : PinnedListView->GetTreeViewItems())
+								{
+									Item->SetObjectVisibility(bIsRowSolo);
+								}
+
+								if (bIsRowSolo)
+								{
+									PinnedItem->ClearSoloRow();
+								}
+								else
+								{
+									PinnedItem->SetObjectVisibility(true);
+									PinnedItem->SetThisAsSoloRow();
+								}
+
+								return FReply::Handled();
+							}
+
+							return FReply::Unhandled();
+						}
+					)
+				]
+			;
+	}
+	
+	TSharedPtr<SWidget> Widget;
+	if (UObject* ObjectRef = PinnedItem->GetObject())
+	{
+		FPropertyEditorModule& PropertyEditorModule = FModuleManager::LoadModuleChecked<FPropertyEditorModule>("PropertyEditor");
+		{
+			FSinglePropertyParams Params;
+			Params.NamePlacement = EPropertyNamePlacement::Hidden;
+
+			if (const TSharedPtr<ISinglePropertyView> SinglePropertyView =
+				PropertyEditorModule.CreateSingleProperty(ObjectRef, InColumnName, Params
+			))
+			{
+				if (const TSharedPtr<IPropertyHandle> Handle = SinglePropertyView->GetPropertyHandle())
+				{
+					if (const FProperty* Property = Handle->GetProperty())
+					{
+						// Do this for all selected rows with a similar property
+						FSimpleDelegate OnPropertyValueChanged =
+							FSimpleDelegate::CreateRaw(
+								this,
+								&SObjectMixerEditorListRow::OnPropertyChanged, Property, (void*)ObjectRef);
+					
+						SinglePropertyView->SetOnPropertyValueChanged(OnPropertyValueChanged);
+
+						Widget = SinglePropertyView;
+					}
+				}
+			}
+		}
+	}
+	
+	if (Widget.IsValid())
+	{
+		return SNew(SBox)
+				.Visibility(EVisibility::SelfHitTestInvisible)
+				.HAlign(HAlign_Center)
+				.VAlign(VAlign_Center)
+				[
+					Widget.ToSharedRef()
+				];
+	}
+
+	return SNullWidget::NullWidget;
+}
+
+void SObjectMixerEditorListRow::OnPropertyChanged(const FProperty* Property, void* ContainerWithChangedProperty)
+{
+	if (Property && ContainerWithChangedProperty)
+	{
+		if (const TSharedPtr<FObjectMixerEditorListRow> PinnedItem = Item.Pin())
+		{
+			if (PinnedItem->GetIsSelected())
+			{
+				if (const void* ValuePtr = Property->ContainerPtrToValuePtr<void>(ContainerWithChangedProperty))
+				{
+					FScopedTransaction Transaction( LOCTEXT("PropertyChanged", "Object Mixer - Property Changed") );
+					
+					for (const TSharedPtr<FObjectMixerEditorListRow>& SelectedRow : PinnedItem->GetSelectedTreeViewItems())
+					{
+						if (UObject* SelectedRowObject = SelectedRow->GetObject())
+						{
+							Property->SetValue_InContainer(SelectedRowObject, ValuePtr);
+						}
+					}
+				}
+			}
+		}
+	}
+}
+
+#undef LOCTEXT_NAMESPACE
diff --git a/Engine/Source/Editor/ObjectMixer/Private/Views/MainPanel/ObjectMixerEditorMainPanel.cpp b/Engine/Source/Editor/ObjectMixer/Private/Views/MainPanel/ObjectMixerEditorMainPanel.cpp
new file mode 100644
--- /dev/null
+++ b/Engine/Source/Editor/ObjectMixer/Private/Views/MainPanel/ObjectMixerEditorMainPanel.cpp
@@ -0,0 +1,62 @@
+// Copyright Epic Games, Inc. All Rights Reserved.
+
+#include "Views/MainPanel/ObjectMixerEditorMainPanel.h"
+
+#include "ObjectMixerEditorModule.h"
+#include "ObjectMixerEditorProjectSettings.h"
+#include "Views/List/ObjectMixerEditorList.h"
+#include "Views/MainPanel/SObjectMixerEditorMainPanel.h"
+
+FObjectMixerEditorMainPanel::FObjectMixerEditorMainPanel()
+{
+	EditorList = MakeShared<FObjectMixerEditorList>();
+}
+
+TSharedRef<SWidget> FObjectMixerEditorMainPanel::GetOrCreateWidget()
+{
+	if (!MainPanelWidget.IsValid())
+	{
+		SAssignNew(MainPanelWidget, SObjectMixerEditorMainPanel, SharedThis(this));
+	}
+
+	return MainPanelWidget.ToSharedRef();
+}
+
+void FObjectMixerEditorMainPanel::RebuildList(const FString InItemToScrollTo, bool bShouldCacheValues) const
+{
+	if (EditorList.IsValid())
+	{
+		EditorList->RebuildList(InItemToScrollTo);
+	}
+}
+
+void FObjectMixerEditorMainPanel::RefreshList() const
+{
+	if (EditorList.IsValid())
+	{
+		EditorList->RefreshList();
+	}
+}
+
+void FObjectMixerEditorMainPanel::OnClassSelectionChanged(UClass* InNewClass) const
+{
+	if (TSharedPtr<FObjectMixerEditorList, ESPMode::ThreadSafe> PinnedList = GetEditorList().Pin())
+	{
+		PinnedList->SetObjectFilterClass(InNewClass);
+	}
+}
+
+TObjectPtr<UClass> FObjectMixerEditorMainPanel::GetClassSelection() const
+{
+	if (TSharedPtr<FObjectMixerEditorList, ESPMode::ThreadSafe> PinnedList = GetEditorList().Pin())
+	{
+		return PinnedList->GetObjectFilterClass();
+	}
+
+	return nullptr;
+}
+
+bool FObjectMixerEditorMainPanel::IsClassSelected(UClass* InNewClass) const
+{
+	return InNewClass == GetClassSelection();
+}
diff --git a/Engine/Source/Editor/ObjectMixer/Private/Views/MainPanel/SObjectMixerEditorMainPanel.cpp b/Engine/Source/Editor/ObjectMixer/Private/Views/MainPanel/SObjectMixerEditorMainPanel.cpp
new file mode 100644
--- /dev/null
+++ b/Engine/Source/Editor/ObjectMixer/Private/Views/MainPanel/SObjectMixerEditorMainPanel.cpp
@@ -0,0 +1,175 @@
+// Copyright Epic Games, Inc. All Rights Reserved.
+
+#include "Views/MainPanel/SObjectMixerEditorMainPanel.h"
+
+#include "ObjectMixerEditorLog.h"
+#include "Views/List/ObjectMixerEditorList.h"
+#include "Views/MainPanel/ObjectMixerEditorMainPanel.h"
+
+#include "AssetRegistry/AssetRegistryModule.h"
+#include "ContentBrowserModule.h"
+#include "Framework/MultiBox/MultiBoxBuilder.h"
+#include "IContentBrowserSingleton.h"
+#include "ISinglePropertyView.h"
+#include "Kismet2/SClassPickerDialog.h"
+#include "PropertyEditorModule.h"
+#include "SPositiveActionButton.h"
+#include "Styling/StyleColors.h"
+
+#include "Widgets/Images/SImage.h"
+#include "Widgets/Input/SComboButton.h"
+#include "Widgets/Input/SMultiLineEditableTextBox.h"
+#include "Widgets/Layout/SSplitter.h"
+
+#define LOCTEXT_NAMESPACE "ObjectMixerEditor"
+
+void SObjectMixerEditorMainPanel::Construct(
+	const FArguments& InArgs, const TSharedRef<FObjectMixerEditorMainPanel>& InMainPanel)
+{
+	check(InMainPanel->GetEditorList().IsValid());
+
+	MainPanel = InMainPanel;
+	
+	ChildSlot
+	[
+		SNew(SVerticalBox)
+			
+		+SVerticalBox::Slot()
+		.HAlign(HAlign_Fill)
+		.AutoHeight()
+		.Padding(FMargin(8.f, 0.f, 8.f, 0.f))
+		[
+			GeneratePanelToolbar()
+		]
+
+		+SVerticalBox::Slot()
+		[
+			SNew(SSplitter)
+			.Orientation(Orient_Vertical)
+			.ResizeMode(ESplitterResizeMode::FixedSize)
+
+			+SSplitter::Slot()
+			[
+				MainPanel.Pin()->GetEditorList().Pin()->GetOrCreateWidget()
+			]
+		]
+	];
+}
+
+SObjectMixerEditorMainPanel::~SObjectMixerEditorMainPanel()
+{
+	MainPanel.Reset();
+	ToolbarHBox.Reset();
+	ConcertButtonPtr.Reset();
+}
+
+TSharedRef<SWidget> SObjectMixerEditorMainPanel::GeneratePanelToolbar()
+{
+	FPropertyEditorModule& PropertyEditorModule = FModuleManager::LoadModuleChecked<FPropertyEditorModule>("PropertyEditor");
+
+	FSinglePropertyParams Params;
+	Params.NamePlacement = EPropertyNamePlacement::Hidden;
+	
+	return SAssignNew(ToolbarHBox, SHorizontalBox)
+				
+		// Add Console Variable button
+		+ SHorizontalBox::Slot()
+		.VAlign(VAlign_Fill)
+		.HAlign(HAlign_Left)
+		.Padding(FMargin(0, 4))
+		[
+			SNew(SPositiveActionButton)
+			.Text(LOCTEXT("AddObject", "Add"))
+			.OnClicked_Lambda([](){ return FReply::Handled(); })
+		]
+
+		// Filter Class Management Button
+		+ SHorizontalBox::Slot()
+		.VAlign(VAlign_Center)
+		.HAlign(HAlign_Right)
+		.Padding(FMargin(8.f, 4, 0.f, 4))
+		[
+			SNew(SComboButton)
+			.ToolTipText(LOCTEXT("FilterClassManagementButton_Tooltip", "Select a filter class"))
+			.ContentPadding(FMargin(4, 0.5f))
+			.ComboButtonStyle(&FAppStyle::Get().GetWidgetStyle<FComboButtonStyle>("ComboButton"))
+			.OnGetMenuContent(this, &SObjectMixerEditorMainPanel::OnGeneratePresetsMenu)
+			.ForegroundColor(FStyleColors::Foreground)
+			.ButtonContent()
+			[
+				SNew(SHorizontalBox)
+
+				+ SHorizontalBox::Slot()
+				.Padding(0, 1, 4, 0)
+				.AutoWidth()
+				[
+					SNew(SImage)
+					.Image(FAppStyle::Get().GetBrush("Icons.Filter"))
+					.ColorAndOpacity(FSlateColor::UseForeground())
+				]
+
+				+ SHorizontalBox::Slot()
+				.Padding(0, 1, 0, 0)
+				.AutoWidth()
+				[
+					SNew(STextBlock)
+					.Text(LOCTEXT("FilterClassToolbarButton", "Object Filter Class"))
+				]
+			]
+		];
+}
+
+TSharedRef<SWidget> SObjectMixerEditorMainPanel::OnGeneratePresetsMenu()
+{
+	FMenuBuilder MenuBuilder(true, NULL);
+
+	TArray<UClass*> DerivedClasses;
+	GetDerivedClasses(UObjectMixerObjectFilter::StaticClass(), DerivedClasses, true);
+
+	DerivedClasses.Remove(UObjectMixerObjectFilter::StaticClass());
+	DerivedClasses.Remove(UObjectMixerBlueprintObjectFilter::StaticClass());
+
+	if (DerivedClasses.Num())
+	{
+		const TSharedPtr<FObjectMixerEditorMainPanel, ESPMode::ThreadSafe> PinnedPanel = MainPanel.Pin();
+		check(PinnedPanel);
+		
+		MenuBuilder.BeginSection(NAME_None, LOCTEXT("SelectClassMenuSection", "Select Class"));
+		{
+			for (UClass* DerivedClass : DerivedClasses)
+			{
+				MenuBuilder.AddMenuEntry(
+				FText::FromName(DerivedClass->GetFName()),
+					FText::GetEmpty(),
+					FSlateIcon(),
+					FUIAction(
+						FExecuteAction::CreateSP(PinnedPanel.ToSharedRef(), &FObjectMixerEditorMainPanel::OnClassSelectionChanged, DerivedClass),
+						FCanExecuteAction::CreateLambda([](){ return true; }),
+						FIsActionChecked::CreateSP(PinnedPanel.ToSharedRef(), &FObjectMixerEditorMainPanel::IsClassSelected, DerivedClass)
+					),
+					NAME_None,
+					EUserInterfaceActionType::RadioButton
+				);
+			}
+		}
+		MenuBuilder.EndSection();
+	}
+	else
+	{
+		MenuBuilder.AddMenuEntry(LOCTEXT("NoFilterClassesAvailable", "No filter classes available."), FText::GetEmpty(), FSlateIcon(), FUIAction());
+	}
+
+	TSharedRef<SWidget> Widget = MenuBuilder.MakeWidget();
+	FChildren* ChildWidgets = Widget->GetChildren();
+	for (int32 ChildItr = 0; ChildItr < ChildWidgets->Num(); ChildItr++)
+	{
+		TSharedRef<SWidget> Child = ChildWidgets->GetChildAt(ChildItr);
+
+		Child->EnableToolTipForceField(false);
+	}
+	Widget->EnableToolTipForceField(false);
+	
+	return Widget;
+}
+
+#undef LOCTEXT_NAMESPACE
diff --git a/Engine/Source/Editor/ObjectMixer/Public/ObjectFilter/ObjectMixerEditorObjectFilter.h b/Engine/Source/Editor/ObjectMixer/Public/ObjectFilter/ObjectMixerEditorObjectFilter.h
new file mode 100644
--- /dev/null
+++ b/Engine/Source/Editor/ObjectMixer/Public/ObjectFilter/ObjectMixerEditorObjectFilter.h
@@ -0,0 +1,250 @@
+// Copyright Epic Games, Inc. All Rights Reserved.
+
+#pragma once
+
+#include "Components/LightComponent.h"
+#include "UObject/Class.h"
+
+#include "ObjectMixerEditorObjectFilter.generated.h"
+
+UENUM(BlueprintType)
+enum class EObjectMixerPropertyInheritanceInclusionOptions : uint8
+{
+	None, // Get only the properties in the specified classes without considering parent or child classes + Specified Class
+	IncludeOnlyImmediateParent, // Get properties from the class the specified class immediately derives from, but not their parents + Specified Class
+	IncludeOnlyImmediateChildren, // Get properties from child classes but not child classes of child classes + Specified Class
+	IncludeOnlyImmediateParentAndChildren, // IncludeOnlyImmediateParent + IncludeOnlyImmediateChildren + Specified Class
+	IncludeAllParents, // Go up the chain of Super classes to get all properties in the class' ancestry + Specified Class
+	IncludeAllChildren, // Get properties from all derived classes recursively + Specified Class
+	IncludeAllParentsAndChildren, // IncludeAllParents + IncludeAllChildren + Specified Class
+	IncludeAllParentsAndOnlyImmediateChildren, // IncludeAllParents + IncludeOnlyImmediateChildren + Specified Class
+	IncludeOnlyImmediateParentAndAllChildren // IncludeOnlyImmediateParent + IncludeAllChildren + Specified Class
+};
+
+/**
+ * Native class for filtering object types to Object Mixer.
+ * Native C++ classes should inherit directly from this class.
+ */
+UCLASS(Abstract, Blueprintable)
+class OBJECTMIXEREDITOR_API UObjectMixerObjectFilter : public UObject
+{
+	GENERATED_BODY()
+public:
+	
+	/*
+	 * Return the basic object types you want to filter for in your level.
+	 * For example, if you want to work with Lights, return ULightComponentBase.
+	 * If you also want to see the properties for parent or child classes,
+	 * override the GetObjectMixerPropertyInheritanceInclusionOptions and GetForceAddedColumns functions.
+	 */
+	virtual TArray<UClass*> GetObjectClassesToFilter() const { return {}; }
+
+	/*
+	 * Get the text to display for the object name/label.
+	 * This is useful if one of your classes is a component type and you want the label of the component's owning actor, for example.
+	 * If not overridden, this returns the object's name.
+	 */
+	virtual FText GetRowDisplayName(UObject* InObject) const;
+
+	/*
+	 * Controls how to display the row's visibility icon. Return true if the object should be visible.
+	 * Generally this should work like the Scene Outliner does.
+	 * If not overridden, we use the Editor Visibility of the object's AActor outer (unless it's an actor itself).
+	 */
+	virtual bool GetRowEditorVisibility(UObject* InObject) const;
+
+	/*
+	 * Controls what happens when the row's visibility icon is clicked.
+	 * Generally this should work like the Scene Outliner does.
+	 * If not overridden, we set the Editor Visibility of the object's AActor outer (unless it's an actor itself).
+	 */
+	virtual void OnSetRowEditorVisibility(UObject* InObject, bool bNewIsVisible) const;
+
+	/*
+	 * Specify a list of property names corresponding to columns you want to show by default.
+	 * For example, you can specify "Intensity" and "LightColor" to show only those property columns by default in the UI.
+	 * Columns not specified will not be shown by default but can be enabled by the user in the UI.
+	 */
+	virtual TArray<FName> GetColumnsToShowByDefault() const;
+
+	/*
+	 * Specify a list of property names corresponding to columns you don't want to ever show.
+	 * For example, you can specify "Intensity" and "LightColor" to ensure that they can't be enabled or shown in the UI.
+	 * Columns not specified can be enabled by the user in the UI.
+	 */
+	virtual TArray<FName> GetColumnsFilter() const;
+
+	/*
+	 * Specify a list of property names found in parent classes you want to show that aren't in the specified classes.
+	 * Note that properties specified here do not override the properties specified in GetColumnsFilter(),
+	 * but do override the supported property tests so these will appear even if ShouldIncludeUnsupportedProperties returns false.
+	 * For example, a ULightComponent displays "LightColor" in the editor's details panel,
+	 * but ULightComponent itself doesn't have a property named "LightColor". Instead it's in its parent class, ULightComponentBase.
+	 * In this scenario, ULightComponent is specified and PropertyInheritanceInclusionOptions is None, so "LightColor" won't appear by default.
+	 * Specify "LightColor" in this function to ensure that "LightColor" will appear as a column as long as
+	 * the property is accessible to one of the specified classes regardless of which parent class it comes from.
+	 */
+	virtual TArray<FName> GetForceAddedColumns() const { return {}; }
+
+	/*
+	 * Specify whether we should return only the properties of the specified classes or the properties of parent and child classes.
+	 * Defaults to 'None' which only considers the properties of the specified classes.
+	 * If you're not seeing all the properties you expected, try overloading this function.
+	 */
+	virtual EObjectMixerPropertyInheritanceInclusionOptions GetObjectMixerPropertyInheritanceInclusionOptions() const;
+
+	/*
+	 * If true, properties that are not visible in the details panel and properties not supported by SSingleProperty will be selectable.
+	 * Defaults to false.
+	 */
+	virtual bool ShouldIncludeUnsupportedProperties() const;
+
+	static TArray<UClass*> GetParentAndChildClassesFromSpecifiedClasses(
+		const TArray<UClass*>& InSpecifiedClasses, EObjectMixerPropertyInheritanceInclusionOptions Options);
+
+	static EFieldIterationFlags GetDesiredFieldIterationFlags(const bool bIncludeInheritedProperties);
+
+	// Folder view should reflect outliner folder structure if possible
+	// Mute/solo lights
+
+protected:
+
+	/*
+	 * Given a set of property names you wish to include, returns a list of all other properties on InObject not found in ExcludeList.
+	 * Useful when defining default visible columns in a list view.
+	 */
+	TArray<FName> GenerateIncludeListFromExcludeList(const TSet<FName>& ExcludeList) const;
+};
+
+/**
+ * Script class for filtering object types to Object Mixer.
+ * Blueprint classes should inherit directly from this class.
+ */
+UCLASS(Abstract, Blueprintable)
+class OBJECTMIXEREDITOR_API UObjectMixerBlueprintObjectFilter : public UObjectMixerObjectFilter
+{
+	GENERATED_BODY()
+public:
+	
+	/*
+	 * Return the basic object types you want to filter for in your level.
+	 * For example, if you want to work with Lights, return ULightComponentBase.
+	 * If you also want to see the properties for parent or child classes,
+	 * override the GetObjectMixerPropertyInheritanceInclusionOptions and GetForceAddedColumns functions.
+	 */
+	UFUNCTION(BlueprintNativeEvent, BlueprintCallable, Category = "Object Mixer")
+	TArray<UClass*> GetObjectClassesToFilter() const override;
+
+	TArray<UClass*> GetObjectClassesToFilter_Implementation() const
+	{
+		return Super::GetObjectClassesToFilter();
+	}
+
+	/*
+	 * Get the text to display for the object name/label.
+	 * This is useful if one of your classes is a component type and you want the label of the component's owning actor, for example.
+	 * If not overridden, this returns the object's name.
+	 */
+	UFUNCTION(BlueprintNativeEvent, BlueprintCallable, Category = "Object Mixer")
+	FText GetRowDisplayName(UObject* InObject) const override;
+	
+	FText GetRowDisplayName_Implementation(UObject* InObject) const
+	{
+		return Super::GetRowDisplayName(InObject);
+	}
+	
+	/*
+	 * Controls how to display the row's visibility icon. Return true if the object should be visible.
+	 * Generally this should work like the Scene Outliner does.
+	 * If not overridden, we use the Editor Visibility of the object's AActor outer (unless it's an actor itself).
+	 */
+	UFUNCTION(BlueprintNativeEvent, BlueprintCallable, Category = "Object Mixer")
+	bool GetRowEditorVisibility(UObject* InObject) const override;
+
+	bool GetRowEditorVisibility_Implementation(UObject* InObject) const
+	{
+		return Super::GetRowEditorVisibility(InObject);
+	}
+	
+	/*
+	 * Controls what happens when the row's visibility icon is clicked.
+	 * Generally this should work like the Scene Outliner does.
+	 * If not overridden, we set the Editor Visibility of the object's AActor outer (unless it's an actor itself).
+	 */
+	UFUNCTION(BlueprintNativeEvent, BlueprintCallable, Category = "Object Mixer")
+	void OnSetRowEditorVisibility(UObject* InObject, bool bNewIsVisible) const override;
+
+	void OnSetRowEditorVisibility_Implementation(UObject* InObject, bool bNewIsVisible) const
+	{
+		return Super::OnSetRowEditorVisibility(InObject, bNewIsVisible);
+	}
+	
+	/*
+	 * Specify a list of property names corresponding to columns you want to show by default.
+	 * For example, you can specify "Intensity" and "LightColor" to show only those property columns by default in the UI.
+	 * Columns not specified will not be shown by default but can be enabled by the user in the UI.
+	 */
+	UFUNCTION(BlueprintNativeEvent, BlueprintCallable, Category = "Object Mixer")
+	TArray<FName> GetColumnsToShowByDefault() const override;
+	
+	TArray<FName> GetColumnsToShowByDefault_Implementation() const
+	{
+		return Super::GetColumnsToShowByDefault();
+	}
+	
+	/*
+	 * Specify a list of property names corresponding to columns you don't want to ever show.
+	 * For example, you can specify "Intensity" and "LightColor" to ensure that they can't be enabled or shown in the UI.
+	 * Columns not specified can be enabled by the user in the UI.
+	 */
+	UFUNCTION(BlueprintNativeEvent, BlueprintCallable, Category = "Object Mixer")
+	TArray<FName> GetColumnsFilter() const override;
+
+	TArray<FName> GetColumnsFilter_Implementation() const
+	{
+		return Super::GetColumnsFilter();
+	}
+
+	/*
+	 * Specify a list of property names found in parent classes you want to show that aren't in the specified classes.
+	 * Note that properties specified here do not override the properties specified in GetColumnsFilter().
+	 * For example, a ULightComponent displays "LightColor" in the editor's details panel,
+	 * but ULightComponent itself doesn't have a property named "LightColor". Instead it's in its parent class, ULightComponentBase.
+	 * In this scenario, ULightComponent is specified and PropertyInheritanceInclusionOptions is None, so "LightColor" won't appear by default.
+	 * Specify "LightColor" in this function to ensure that "LightColor" will appear as a column as long as
+	 * the property is accessible to one of the specified classes regardless of which parent class it comes from.
+	 */
+	UFUNCTION(BlueprintNativeEvent, BlueprintCallable, Category = "Object Mixer")
+	TArray<FName> GetForceAddedColumns() const override;
+	
+	TArray<FName> GetForceAddedColumns_Implementation() const
+	{
+		return Super::GetForceAddedColumns();
+	}
+	
+	/*
+	 * Specify whether we should return only the properties of the specified classes or the properties of parent and child classes.
+	 * Defaults to 'None' which only considers the properties of the specified classes.
+	 * If you're not seeing all the properties you expected, try overloading this function.
+	 */
+	UFUNCTION(BlueprintNativeEvent, BlueprintCallable, Category = "Object Mixer")
+	EObjectMixerPropertyInheritanceInclusionOptions GetObjectMixerPropertyInheritanceInclusionOptions() const override;
+
+	EObjectMixerPropertyInheritanceInclusionOptions GetObjectMixerPropertyInheritanceInclusionOptions_Implementation() const
+	{
+		return Super::GetObjectMixerPropertyInheritanceInclusionOptions();
+	}
+
+	/*
+	 * If true, properties that are not visible in the details panel and properties not supported by SSingleProperty will be selectable.
+	 * Defaults to false.
+	 */
+	UFUNCTION(BlueprintNativeEvent, BlueprintCallable, Category = "Object Mixer")
+	bool ShouldIncludeUnsupportedProperties() const override;
+
+	bool ShouldIncludeUnsupportedProperties_Implementation() const
+	{
+		return Super::ShouldIncludeUnsupportedProperties();
+	}
+
+};
diff --git a/Engine/Source/Editor/ObjectMixer/Public/ObjectMixerEditorLog.h b/Engine/Source/Editor/ObjectMixer/Public/ObjectMixerEditorLog.h
new file mode 100644
--- /dev/null
+++ b/Engine/Source/Editor/ObjectMixer/Public/ObjectMixerEditorLog.h
@@ -0,0 +1,7 @@
+// Copyright Epic Games, Inc. All Rights Reserved.
+
+#pragma once
+
+#include "Logging/LogCategory.h"
+
+DECLARE_LOG_CATEGORY_CLASS(LogObjectMixerEditor, Log, All);
diff --git a/Engine/Source/Editor/ObjectMixer/Public/ObjectMixerEditorModule.h b/Engine/Source/Editor/ObjectMixer/Public/ObjectMixerEditorModule.h
new file mode 100644
--- /dev/null
+++ b/Engine/Source/Editor/ObjectMixer/Public/ObjectMixerEditorModule.h
@@ -0,0 +1,67 @@
+// Copyright Epic Games, Inc. All Rights Reserved.
+
+#pragma once
+
+#include "Modules/ModuleInterface.h"
+#include "ObjectFilter/ObjectMixerEditorObjectFilter.h"
+#include "Widgets/Docking/SDockTab.h"
+
+class FObjectMixerEditorMainPanel;
+
+class OBJECTMIXEREDITOR_API FObjectMixerEditorModule : public IModuleInterface
+{
+public:
+
+	//~ Begin IModuleInterface Interface
+	virtual void StartupModule() override;
+	virtual void ShutdownModule() override;
+	//~ End IModuleInterface 
+
+	static FObjectMixerEditorModule& Get();
+	
+	virtual TSharedPtr<SWidget> MakeObjectMixerDialog() const;
+
+	/*
+	 * Regenerate the list items and refresh the list. Call when adding or removing variables.
+	 * @param bShouldCacheValues If true, the current list's current values will be cached and then restored when the list is rebuilt. Otherwise preset values will be used.
+	 */
+	virtual void RebuildList(const FString InItemToScrollTo = "", bool bShouldCacheValues = true) const;
+	
+	/**
+	 * Refresh filters and sorting.
+	 * Useful for when the list state has gone stale but the variable count has not changed.
+	 */
+	virtual void RefreshList() const;
+
+	void RegisterMenuGroupAndTabSpawner();
+	/**
+	 * Add a tab spawner to the Object Mixer menu group.
+	 * @return If adding the item to the menu was successful
+	 */
+	bool RegisterItemInMenuGroup(FWorkspaceItem& InItem);
+	
+	virtual void UnregisterTabSpawner();
+	virtual void RegisterProjectSettings() const;
+	virtual void UnregisterProjectSettings() const;
+
+	virtual TSharedRef<SDockTab> SpawnMainPanelTab(const FSpawnTabArgs& Args);
+	
+	static const FName ObjectMixerToolkitPanelTabId;
+
+	TSharedPtr<FWorkspaceItem> GetWorkspaceGroup();
+
+protected:
+	
+	/** Lives for as long as the module is loaded. */
+	TSharedPtr<FObjectMixerEditorMainPanel> MainPanel;
+
+	/** The text that appears on the spawned nomad tab */
+	FText TabLabel;
+
+	// If set, this is the filter class used to initialize the MainPanel
+	TSubclassOf<UObjectMixerObjectFilter> DefaultFilterClass;
+
+private:
+
+	TSharedPtr<FWorkspaceItem> WorkspaceGroup;
+};
diff --git a/Engine/Source/Editor/ObjectMixer/Public/ObjectMixerEditorProjectSettings.h b/Engine/Source/Editor/ObjectMixer/Public/ObjectMixerEditorProjectSettings.h
new file mode 100644
--- /dev/null
+++ b/Engine/Source/Editor/ObjectMixer/Public/ObjectMixerEditorProjectSettings.h
@@ -0,0 +1,23 @@
+// Copyright Epic Games, Inc. All Rights Reserved.
+
+#pragma once
+
+#include "ObjectFilter/ObjectMixerEditorObjectFilter.h"
+
+#include "ObjectMixerEditorProjectSettings.generated.h"
+
+UCLASS(config = Engine, defaultconfig)
+class OBJECTMIXEREDITOR_API UObjectMixerEditorProjectSettings : public UObject
+{
+	GENERATED_BODY()
+public:
+	
+	UObjectMixerEditorProjectSettings(const FObjectInitializer& ObjectInitializer)
+	{}
+
+	/**
+	 * If false, a new object will be created every time the filter object is accessed.
+	 */
+	UPROPERTY(Config, EditAnywhere, Category="Object Mixer")
+	bool bShouldCacheFilterObject = true;
+};
diff --git a/Engine/Source/Editor/ObjectMixer/Public/ObjectMixerEditorStyle.h b/Engine/Source/Editor/ObjectMixer/Public/ObjectMixerEditorStyle.h
new file mode 100644
--- /dev/null
+++ b/Engine/Source/Editor/ObjectMixer/Public/ObjectMixerEditorStyle.h
@@ -0,0 +1,36 @@
+// Copyright Epic Games, Inc. All Rights Reserved.
+
+#pragma once
+
+#include "Styling/SlateStyle.h"
+
+class OBJECTMIXEREDITOR_API FObjectMixerEditorStyle : FSlateStyleSet
+{
+public:
+
+	static void Initialize();
+
+	static void Shutdown();
+
+	static void ReloadTextures();
+
+	static const ISlateStyle& Get();
+
+	virtual const FName& GetStyleSetName() const override;
+
+	virtual const FSlateBrush* GetBrush(const FName PropertyName, const ANSICHAR* Specifier = nullptr, const ISlateStyle* RequestingStyle = nullptr) const override;
+
+	template< typename WidgetStyleType >
+	static const WidgetStyleType& GetWidgetStyle(FName PropertyName, const ANSICHAR* Specifier = nullptr)
+	{
+		return StyleInstance->GetWidgetStyle<WidgetStyleType>(PropertyName, Specifier);
+	}
+
+private:
+
+	static FString GetExternalPluginContent(const FString& PluginName, const FString& RelativePath, const ANSICHAR* Extension);
+	
+	static TSharedRef<FSlateStyleSet> Create();
+
+	static TSharedPtr<FSlateStyleSet> StyleInstance;
+};
diff --git a/Engine/Source/Editor/ObjectMixer/Public/ObjectMixerFilterFactory.h b/Engine/Source/Editor/ObjectMixer/Public/ObjectMixerFilterFactory.h
new file mode 100644
--- /dev/null
+++ b/Engine/Source/Editor/ObjectMixer/Public/ObjectMixerFilterFactory.h
@@ -0,0 +1,27 @@
+// Copyright Epic Games, Inc. All Rights Reserved.
+
+#pragma once
+
+#include "Factories/Factory.h"
+
+#include "ObjectMixerFilterFactory.generated.h"
+
+UCLASS(hidecategories = Object)
+class OBJECTMIXEREDITOR_API UObjectMixerBlueprintFilterFactory : public UFactory
+{
+	GENERATED_BODY()
+public:
+	UObjectMixerBlueprintFilterFactory();
+	
+	// The parent class of the created blueprint
+    UPROPERTY(VisibleAnywhere, Category="ObjectMixerBlueprintFilterFactory", meta=(AllowAbstract = "", BlueprintBaseOnly = ""))
+    TSubclassOf<UObject> ParentClass;
+	
+	//~ Begin UFactory Interface
+	virtual FText GetDisplayName() const override;
+	virtual UObject* FactoryCreateNew(UClass* Class, UObject* InParent, FName Name, EObjectFlags Flags, UObject* Context, FFeedbackContext* Warn) override;
+	virtual FString GetDefaultNewAssetName() const override;
+	virtual uint32 GetMenuCategories() const override;
+	//~ Begin UFactory Interface
+	
+};
diff --git a/Engine/Source/Editor/ObjectMixer/Public/Views/List/ObjectMixerEditorList.h b/Engine/Source/Editor/ObjectMixer/Public/Views/List/ObjectMixerEditorList.h
new file mode 100644
--- /dev/null
+++ b/Engine/Source/Editor/ObjectMixer/Public/Views/List/ObjectMixerEditorList.h
@@ -0,0 +1,106 @@
+// Copyright Epic Games, Inc. All Rights Reserved.
+
+#pragma once
+
+#include "ObjectFilter/ObjectMixerEditorObjectFilter.h"
+#include "ObjectMixerEditorProjectSettings.h"
+#include "Widgets/SWidget.h"
+
+struct FObjectMixerEditorListRow;
+typedef TSharedPtr<FObjectMixerEditorListRow> FObjectMixerEditorListRowPtr;
+
+class SObjectMixerEditorList;
+
+class OBJECTMIXEREDITOR_API FObjectMixerEditorList : public TSharedFromThis<FObjectMixerEditorList>
+{
+public:
+
+	FObjectMixerEditorList();
+
+	virtual ~FObjectMixerEditorList();
+
+	TSharedRef<SWidget> GetOrCreateWidget();
+
+	UObjectMixerObjectFilter* GetObjectFilter();
+
+	void CacheObjectFilterObject();
+
+	/*
+	 * Returns either ObjectClassFromWhichToGeneratePropertyColumns or ObjectClassOverride, if set.
+	 * @param bForceNoOverride If true, always return ObjectClassFromWhichToGeneratePropertyColumns and not ObjectClassOverride.
+	 */
+	TArray<UClass*> GetObjectClasses(bool bForceNoOverride = false)
+	{
+		TArray<UClass*> ReturnValue;
+
+		if (const UObjectMixerObjectFilter* Filter = GetObjectFilter())
+		{
+			ReturnValue = Filter->GetObjectClassesToFilter();
+		}
+		
+		return ReturnValue;
+	}
+
+	TWeakPtr<FObjectMixerEditorListRow> GetSoloRow()
+	{
+		return SoloRow;
+	}
+
+	void SetSoloRow(TSharedRef<FObjectMixerEditorListRow> InRow)
+	{
+		SoloRow = InRow;
+	}
+
+	void ClearSoloRow()
+	{
+		SoloRow = nullptr;
+	}
+
+	void SetSearchString(const FString& SearchString);
+
+	void ClearList() const;
+
+	/*
+	 * Regenerate the list items and refresh the list. Call when adding or removing variables.
+	 */
+	void RebuildList(const FString& InItemToScrollTo = "") const;
+
+	/**
+	 * Refresh filters and sorting.
+	 * Useful for when the list state has gone stale but the variable count has not changed.
+	 */
+	void RefreshList() const;
+
+	TSubclassOf<UObjectMixerObjectFilter> GetObjectFilterClass() const
+	{
+		return ObjectFilterClass;
+	}
+
+	void SetObjectFilterClass(TSubclassOf<UObjectMixerObjectFilter> InObjectFilterClass)
+	{
+		ObjectFilterClass = InObjectFilterClass;
+		CacheObjectFilterObject();
+		RebuildList();
+	}
+
+private:
+
+	/**
+	 * The class used to generate property edit columns
+	 */
+	TSubclassOf<UObjectMixerObjectFilter> ObjectFilterClass;
+
+	TSharedPtr<SObjectMixerEditorList> ListWidget;
+
+	TWeakObjectPtr<UObjectMixerObjectFilter> ObjectFilterPtr;
+
+	TWeakPtr<FObjectMixerEditorListRow> SoloRow = nullptr;
+
+	// Delegates
+	TFunction<void()> RebuildDelegate = [this]()
+	{
+		RebuildList();
+	};
+	
+	TSet<FDelegateHandle> EditorDelegateHandles;
+};
diff --git a/Engine/Source/Editor/ObjectMixer/Public/Views/List/ObjectMixerEditorListFilters/IObjectMixerEditorListFilter.h b/Engine/Source/Editor/ObjectMixer/Public/Views/List/ObjectMixerEditorListFilters/IObjectMixerEditorListFilter.h
new file mode 100644
--- /dev/null
+++ b/Engine/Source/Editor/ObjectMixer/Public/Views/List/ObjectMixerEditorListFilters/IObjectMixerEditorListFilter.h
@@ -0,0 +1,63 @@
+// Copyright Epic Games, Inc. All Rights Reserved.
+
+#pragma once
+
+#include "../ObjectMixerEditorListRow.h"
+
+class OBJECTMIXEREDITOR_API IObjectMixerEditorListFilter
+{
+public:
+
+	enum class EObjectMixerEditorListFilterMatchType
+	{
+		// Filters of this type are placed into a group in which if any of the group's filters pass, the row is considered passing
+		MatchAny,
+		// Filters of this type are placed into a group in which all of the group's filters must pass for the row to be considered passing
+		MatchAll
+	};
+	
+	virtual ~IObjectMixerEditorListFilter() = default;
+
+	virtual FString GetFilterName() = 0;
+
+	/** Returns localized text to display for when this filter' active state is defined by a toggle button. */
+	virtual FText GetFilterButtonLabel() = 0;
+	/** Returns localized text to display when the user mouses over a toggle button that defines whether this filter is active. */
+	virtual FText GetFilterButtonToolTip() = 0;
+	
+	void SetFilterActive(const bool bNewEnabled)
+	{
+		bIsFilterActive = bNewEnabled;
+	}
+
+	void ToggleFilterActive()
+	{
+		bIsFilterActive = !bIsFilterActive;
+	}
+
+	bool GetIsFilterActive() const
+	{
+		return bIsFilterActive;
+	}
+
+	void SetFilterMatchType(const EObjectMixerEditorListFilterMatchType InType)
+	{
+		FilterMatchType = InType;
+	}
+
+	EObjectMixerEditorListFilterMatchType GetFilterMatchType() const
+	{
+		return FilterMatchType;
+	}
+
+	virtual bool DoesItemPassFilter(const FObjectMixerEditorListRowPtr& InItem)
+	{
+		return false;
+	}
+
+private:
+
+	bool bIsFilterActive = true;
+	EObjectMixerEditorListFilterMatchType FilterMatchType = EObjectMixerEditorListFilterMatchType::MatchAny;
+			
+};
diff --git a/Engine/Source/Editor/ObjectMixer/Public/Views/List/ObjectMixerEditorListFilters/ObjectMixerEditorListFilter_Source.h b/Engine/Source/Editor/ObjectMixer/Public/Views/List/ObjectMixerEditorListFilters/ObjectMixerEditorListFilter_Source.h
new file mode 100644
--- /dev/null
+++ b/Engine/Source/Editor/ObjectMixer/Public/Views/List/ObjectMixerEditorListFilters/ObjectMixerEditorListFilter_Source.h
@@ -0,0 +1,44 @@
+// Copyright Epic Games, Inc. All Rights Reserved.
+
+#pragma once
+
+#include "IObjectMixerEditorListFilter.h"
+
+#define LOCTEXT_NAMESPACE "ObjectMixerEditor"
+
+class OBJECTMIXEREDITOR_API FObjectMixerEditorListFilter_Source : public IObjectMixerEditorListFilter
+{
+public:
+
+	FObjectMixerEditorListFilter_Source(){}
+
+	virtual FString GetFilterName() override
+	{
+		return "";
+	}
+
+	virtual FText GetFilterButtonLabel() override
+	{
+		return FText::Format(LOCTEXT("ShowSourceTextFilterFormat", "Show {0}"), FText::FromString(GetFilterName()));
+	}
+
+	virtual FText GetFilterButtonToolTip() override
+	{
+		return FText::Format(
+			LOCTEXT("ShowSourceTextFilterTooltipFormat", "Show rows that have a Source field matching '{0}'"),
+			FText::FromString(GetFilterName()));
+	}
+
+	virtual bool DoesItemPassFilter(const FObjectMixerEditorListRowPtr& InItem) override
+	{
+		if (InItem.IsValid())
+		{
+			
+		}
+
+		return false;
+	}
+
+};
+
+#undef LOCTEXT_NAMESPACE
diff --git a/Engine/Source/Editor/ObjectMixer/Public/Views/List/ObjectMixerEditorListRow.h b/Engine/Source/Editor/ObjectMixer/Public/Views/List/ObjectMixerEditorListRow.h
new file mode 100644
--- /dev/null
+++ b/Engine/Source/Editor/ObjectMixer/Public/Views/List/ObjectMixerEditorListRow.h
@@ -0,0 +1,151 @@
+// Copyright Epic Games, Inc. All Rights Reserved.
+
+#pragma once
+
+#include "Containers/Array.h"
+#include "Containers/UnrealString.h"
+#include "Layout/Visibility.h"
+#include "Templates/SharedPointer.h"
+#include "UObject/Object.h"
+#include "UObject/WeakObjectPtr.h"
+
+class SObjectMixerEditorList;
+class UObjectMixerObjectFilter;
+
+struct FSlateBrush;
+
+struct FObjectMixerEditorListRow;
+typedef TSharedPtr<FObjectMixerEditorListRow> FObjectMixerEditorListRowPtr;
+
+struct OBJECTMIXEREDITOR_API FObjectMixerEditorListRow final : TSharedFromThis<FObjectMixerEditorListRow>
+{
+	enum EObjectMixerEditorListRowType
+	{
+		None,
+		Group,
+		SingleItem
+	};
+
+	~FObjectMixerEditorListRow();
+
+	void FlushReferences();
+	
+	FObjectMixerEditorListRow(
+		const TWeakObjectPtr<UObject> InObject, const EObjectMixerEditorListRowType InRowType, 
+		const TSharedRef<SObjectMixerEditorList>& InListView, 
+		const TWeakPtr<FObjectMixerEditorListRow>& InDirectParentRow)
+	: ObjectRef(InObject)
+	, RowType(InRowType)
+	, ListViewPtr(InListView)
+	, DirectParentRow(InDirectParentRow)
+	{}
+
+	[[nodiscard]] UObject* GetObject()
+	{
+		if (ObjectRef.IsValid())
+		{
+			return ObjectRef.Get();
+		}
+
+		return nullptr;
+	}
+
+	UObjectMixerObjectFilter* GetObjectFilter() const;
+
+	[[nodiscard]] EObjectMixerEditorListRowType GetRowType() const;
+
+	[[nodiscard]] int32 GetChildDepth() const;
+	void SetChildDepth(const int32 InDepth);
+
+	[[nodiscard]] int32 GetSortOrder() const;
+	void SetSortOrder(const int32 InNewOrder);
+
+	TWeakPtr<FObjectMixerEditorListRow> GetDirectParentRow() const;
+	void SetDirectParentRow(const TWeakPtr<FObjectMixerEditorListRow>& InDirectParentRow);
+	
+	/* bHasGeneratedChildren must be true to get actual children. */
+	[[nodiscard]] const TArray<FObjectMixerEditorListRowPtr>& GetChildRows() const;
+	/* bHasGeneratedChildren must be true to get an accurate value. */
+	[[nodiscard]] int32 GetChildCount() const;
+	void SetChildRows(const TArray<FObjectMixerEditorListRowPtr>& InChildRows);
+	void AddToChildRows(const FObjectMixerEditorListRowPtr& InRow);
+	void InsertChildRowAtIndex(const FObjectMixerEditorListRowPtr& InRow, const int32 AtIndex = 0);
+
+	[[nodiscard]] bool GetIsTreeViewItemExpanded() const;
+	void SetIsTreeViewItemExpanded(const bool bNewExpanded);
+	
+	[[nodiscard]] bool GetShouldFlashOnScrollIntoView() const;
+	void SetShouldFlashOnScrollIntoView(const bool bNewShouldFlashOnScrollIntoView);
+
+	[[nodiscard]] bool GetShouldExpandAllChildren() const;
+	void SetShouldExpandAllChildren(const bool bNewShouldExpandAllChildren);
+
+	void ResetToStartupValueAndSource() const;
+
+	/*
+	 *Individual members of InTokens will be considered "AnyOf" or "OR" searches. If SearchTerms contains any individual member it will match.
+	 *Members will be tested for a space character (" "). If a space is found, a subsearch will be run.
+	 *This subsearch will be an "AllOf" or "AND" type search in which all strings, separated by a space, must be found in the search terms.
+	 */
+	bool MatchSearchTokensToSearchTerms(const TArray<FString> InTokens, ESearchCase::Type InSearchCase = ESearchCase::IgnoreCase);
+
+	/* This overload creates tokens from a string first, then calls ExecuteSearchOnChildNodes(const TArray<FString>& Tokens). */
+	void ExecuteSearchOnChildNodes(const FString& SearchString) const;
+	void ExecuteSearchOnChildNodes(const TArray<FString>& Tokens) const;
+
+	[[nodiscard]] bool GetDoesRowPassFilters() const;
+	void SetDoesRowPassFilters(const bool bPass);
+
+	[[nodiscard]] bool GetIsSelected();
+
+	[[nodiscard]] bool ShouldBeVisible() const;
+	[[nodiscard]] EVisibility GetDesiredVisibility() const;
+
+	[[nodiscard]] bool HasVisibleChildren() const
+	{
+		return false;
+	}
+
+	[[nodiscard]] TWeakPtr<SObjectMixerEditorList> GetListViewPtr() const
+	{
+		return ListViewPtr;
+	}
+
+	[[nodiscard]] TArray<FObjectMixerEditorListRowPtr> GetSelectedTreeViewItems() const;
+
+	const FSlateBrush* GetObjectIconBrush();
+
+	bool GetObjectVisibility();
+	void SetObjectVisibility(const bool bNewIsVisible);
+
+	bool IsThisRowSolo() const;
+	void SetThisAsSoloRow();
+
+	void ClearSoloRow();
+
+private:
+
+	TWeakObjectPtr<UObject> ObjectRef;
+	EObjectMixerEditorListRowType RowType = SingleItem;
+	TArray<FObjectMixerEditorListRowPtr> ChildRows;
+	
+	TWeakPtr<SObjectMixerEditorList> ListViewPtr;
+	
+	bool bIsTreeViewItemExpanded = false;
+	bool bShouldFlashOnScrollIntoView = false;
+
+	int32 ChildDepth = 0;
+
+	int32 SortOrder = -1;
+
+	FString CachedSearchTerms;
+
+	bool bDoesRowMatchSearchTerms = true;
+	bool bDoesRowPassFilters = true;
+	
+	bool bIsSelected = false;
+	TWeakPtr<FObjectMixerEditorListRow> DirectParentRow;
+
+	// Used to expand all children on shift+click.
+	bool bShouldExpandAllChildren = false;
+};
diff --git a/Engine/Source/Editor/ObjectMixer/Public/Views/List/SObjectMixerEditorList.h b/Engine/Source/Editor/ObjectMixer/Public/Views/List/SObjectMixerEditorList.h
new file mode 100644
--- /dev/null
+++ b/Engine/Source/Editor/ObjectMixer/Public/Views/List/SObjectMixerEditorList.h
@@ -0,0 +1,237 @@
+// Copyright Epic Games, Inc. All Rights Reserved.
+
+#pragma once
+
+#include "ObjectMixerEditorList.h"
+#include "ObjectMixerEditorListRow.h"
+#include "ObjectMixerEditorListFilters/IObjectMixerEditorListFilter.h"
+
+#include "Widgets/Input/SComboButton.h"
+#include "Widgets/Layout/SGridPanel.h"
+#include "Widgets/SCompoundWidget.h"
+#include "Widgets/Views/SHeaderRow.h"
+#include "Widgets/Views/STreeView.h"
+
+#include "SObjectMixerEditorList.generated.h"
+
+class SWrapBox;
+class FObjectMixerEditorList;
+class SBox;
+class SComboButton;
+class SSearchBox;
+class SHeaderRow;
+
+UENUM()
+enum class EListViewColumnType
+{
+	BuiltIn,
+	PropertyGenerated
+};
+
+USTRUCT()
+struct FListViewColumnInfo
+{
+	GENERATED_BODY()
+	
+	TObjectPtr<FProperty> PropertyRef;
+	
+	UPROPERTY()
+	FName PropertyName;
+
+	UPROPERTY()
+	FText PropertyDisplayText;
+
+	UPROPERTY()
+	EListViewColumnType PropertyType;
+
+	UPROPERTY()
+	FName CategoryName;
+
+	UPROPERTY()
+	bool bIsDesiredForDisplay = false;
+
+	UPROPERTY()
+	bool bCanBeSorted = false;
+
+	UPROPERTY()
+	bool bUseFixedWidth = false;
+
+	UPROPERTY()
+	float FixedWidth = 25.0f;
+
+	UPROPERTY()
+	float FillWidth = 1.0f;
+};
+
+class OBJECTMIXEREDITOR_API SObjectMixerEditorList final : public SCompoundWidget
+{
+
+public:
+
+	SLATE_BEGIN_ARGS(SObjectMixerEditorList)
+	{}
+
+	SLATE_END_ARGS()
+
+	void Construct(const FArguments& InArgs, TSharedRef<FObjectMixerEditorList> ListModel);
+
+	virtual ~SObjectMixerEditorList() override;
+
+	TWeakPtr<FObjectMixerEditorList> GetListModelPtr()
+	{
+		return ListModelPtr;
+	}
+
+	void ClearList()
+	{
+		FlushMemory(false);
+	}
+
+	/*
+	 * Regenerate the list items and refresh the list. Call when adding or removing items.
+	 */
+	void RebuildList(const FString& InItemToScrollTo = "");
+
+	/**
+	 * Refresh filters and sorting.
+	 * Useful for when the list state has gone stale but the variable count has not changed.
+	 */
+	void RefreshList();
+
+	[[nodiscard]] TArray<FObjectMixerEditorListRowPtr> GetSelectedTreeViewItems() const;
+	int32 GetSelectedTreeViewItemCount() const;
+
+	void SetSelectedTreeViewItemActorsEditorVisible(const bool bNewIsVisible);
+
+	bool IsTreeViewItemSelected(TSharedRef<FObjectMixerEditorListRow> Item);
+
+	[[nodiscard]] TArray<FObjectMixerEditorListRowPtr> GetTreeViewItems() const;
+	void SetTreeViewItems(const TArray<FObjectMixerEditorListRowPtr>& InItems);
+
+	[[nodiscard]] int32 GetTreeViewItemCount() const
+	{
+		return TreeViewRootObjects.Num();
+	}
+	
+	TWeakPtr<FObjectMixerEditorListRow> GetSoloRow()
+	{
+		return GetListModelPtr().Pin()->GetSoloRow();
+	}
+
+	void SetSoloRow(TSharedRef<FObjectMixerEditorListRow> InRow)
+	{
+		GetListModelPtr().Pin()->SetSoloRow(InRow);
+	}
+
+	void ClearSoloRow()
+	{
+		GetListModelPtr().Pin()->ClearSoloRow();
+	}
+
+	FString GetSearchStringFromSearchInputField() const;
+	void SetSearchStringInSearchInputField(const FString InSearchString) const;
+	void ExecuteListViewSearchOnAllRows(const FString& SearchString, const bool bShouldRefreshAfterward = true);
+
+	bool DoesTreeViewHaveVisibleChildren() const;
+
+	void SetTreeViewItemExpanded(const TSharedPtr<FObjectMixerEditorListRow>& RowToExpand, const bool bNewExpansion) const;
+
+	void ToggleFilterActive(const FString& FilterName);
+	void EvaluateIfRowsPassFilters(const bool bShouldRefreshAfterward = true);
+
+	// Sorting
+
+	const FName& GetActiveSortingColumnName() const
+	{
+		return ActiveSortingColumnName;
+	}
+	EColumnSortMode::Type GetSortModeForColumn(FName InColumnName) const;
+	void OnSortColumnCalled(EColumnSortPriority::Type Priority, const FName& ColumnName, EColumnSortMode::Type SortMode);
+	EColumnSortMode::Type CycleSortMode(const FName& InColumnName);
+	void ExecuteSort(
+		const FName& InColumnName, const EColumnSortMode::Type InColumnSortMode, const bool bShouldRefreshAfterward = true);
+	void ClearSorting()
+	{
+		ActiveSortingColumnName = NAME_None;
+		ActiveSortingType = EColumnSortMode::None;
+	}
+
+	// Columns
+
+	FListViewColumnInfo* GetColumnInfoByPropertyName(const FName& InPropertyName);
+
+	static const FName ItemNameColumnName;
+	static const FName EditorVisibilityColumnName;
+	static const FName EditorVisibilitySoloColumnName;
+
+private:
+
+	TWeakPtr<FObjectMixerEditorList> ListModelPtr;
+	
+	TSharedPtr<SHeaderRow> HeaderRow;
+	TSharedRef<SWidget> GenerateHeaderRowContextMenu() const;
+
+	/*
+	 * Only adds properties that pass a series of tests, including having only one unique entry in the column list array.
+	 * @param bForceIncludeProperty If true, only Skiplist and Uniqueness tests will be checked, bypassing class, blueprint editability and other requirements.
+	 * @param
+	 */
+	bool AddUniquePropertyColumnsToHeaderRow(
+		FProperty* Property,
+		const bool bForceIncludeProperty = false,
+		const TArray<FName>& PropertySkipList = {}
+	);
+	void AddBuiltinColumnsToHeaderRow();
+	TSharedPtr<SHeaderRow> GenerateHeaderRow();
+	ECheckBoxState HeaderCheckBoxState = ECheckBoxState::Checked;
+
+	void SetupFilters();
+
+	TSharedRef<SWidget> BuildShowOptionsMenu();
+	
+	void FlushMemory(const bool bShouldKeepMemoryAllocated);
+	
+	void SetAllGroupsCollapsed();
+
+	// Search
+	
+	void OnListViewSearchTextChanged(const FText& Text);
+
+	TSharedPtr<SSearchBox> ListSearchBoxPtr;
+	TSharedPtr<SComboButton> ViewOptionsComboButton;
+	
+	TSharedPtr<SBox> ListBoxContainerPtr;
+
+	//  Tree View Implementation
+
+	void GenerateTreeView();
+	void FindVisibleTreeViewObjects();
+	void FindVisibleObjectsAndRequestTreeRefresh();
+	
+	void OnGetRowChildren(FObjectMixerEditorListRowPtr Row, TArray<FObjectMixerEditorListRowPtr>& OutChildren) const;
+	void OnRowChildExpansionChange(FObjectMixerEditorListRowPtr Row, const bool bIsExpanded, const bool bIsRecursive = false) const;
+
+	void SetChildExpansionRecursively(const FObjectMixerEditorListRowPtr& InRow, const bool bNewIsExpanded) const;
+
+	TArray<TSharedRef<IObjectMixerEditorListFilter>> ShowFilters;
+
+	TSharedPtr<STreeView<FObjectMixerEditorListRowPtr>> TreeViewPtr;
+
+	/** All Tree view objects */
+	TArray<FObjectMixerEditorListRowPtr> TreeViewRootObjects;
+	/** Visible Tree view objects, after filters */
+	TArray<FObjectMixerEditorListRowPtr> VisibleTreeViewObjects;
+
+	TArray<FListViewColumnInfo> ListViewColumns;
+
+	// Sorting
+
+	FName ActiveSortingColumnName = NAME_None;
+	EColumnSortMode::Type ActiveSortingType = EColumnSortMode::None;
+
+	TFunctionRef<bool(const FObjectMixerEditorListRowPtr&, const FObjectMixerEditorListRowPtr&)> SortByOrderAscending =
+		[](const FObjectMixerEditorListRowPtr& A, const FObjectMixerEditorListRowPtr& B)
+		{
+			return A->GetSortOrder() < B->GetSortOrder();
+		};
+};
diff --git a/Engine/Source/Editor/ObjectMixer/Public/Views/List/SObjectMixerEditorListRow.h b/Engine/Source/Editor/ObjectMixer/Public/Views/List/SObjectMixerEditorListRow.h
new file mode 100644
--- /dev/null
+++ b/Engine/Source/Editor/ObjectMixer/Public/Views/List/SObjectMixerEditorListRow.h
@@ -0,0 +1,87 @@
+// Copyright Epic Games, Inc. All Rights Reserved.
+
+#pragma once
+
+#include "ObjectMixerEditorListRow.h"
+
+#include "DragAndDrop/DecoratedDragDropOp.h"
+#include "Widgets/Input/SNumericEntryBox.h"
+#include "Widgets/IToolTip.h"
+#include "Widgets/Views/STableRow.h"
+
+class SObjectMixerEditorListValueInput;
+class SObjectMixerEditorListRowHoverWidgets;
+
+class OBJECTMIXEREDITOR_API SObjectMixerEditorListRow : public SMultiColumnTableRow<FObjectMixerEditorListRowPtr>
+{
+public:
+	
+	SLATE_BEGIN_ARGS(SObjectMixerEditorListRow)
+	{}
+
+	SLATE_END_ARGS()
+	
+	void Construct(const FArguments& InArgs, const TSharedRef<STableViewBase>& InOwnerTable, const TWeakPtr<FObjectMixerEditorListRow> InRow);
+	
+	virtual TSharedRef<SWidget> GenerateWidgetForColumn(const FName& InColumnName) override;
+
+	// Begin SWidget
+	virtual void OnMouseEnter(const FGeometry& MyGeometry, const FPointerEvent& MouseEvent) override;
+	virtual void OnMouseLeave(const FPointerEvent& MouseEvent) override;
+	// End SWidget
+	
+	virtual ~SObjectMixerEditorListRow() override;	
+
+	// FReply HandleDragDetected(const FGeometry& MyGeometry, const FPointerEvent& MouseEvent);
+	// void HandleDragLeave(const FDragDropEvent& DragDropEvent);
+	// TOptional<EItemDropZone> HandleCanAcceptDrop(const FDragDropEvent& DragDropEvent, EItemDropZone DropZone, FObjectMixerEditorListRowPtr TargetItem);
+	// FReply HandleAcceptDrop(const FDragDropEvent& DragDropEvent, EItemDropZone DropZone, FObjectMixerEditorListRowPtr TargetItem);
+	
+	void FlashRow();
+
+private:
+
+	bool IsVisible() const;
+
+	FSlateColor GetVisibilityIconForegroundColor() const;
+	FSlateColor GetSoloIconForegroundColor() const;
+	
+	/** Get the brush for this widget */
+	const FSlateBrush* GetVisibilityBrush() const;
+	
+	EVisibility GetFlashImageVisibility() const;
+	FSlateColor GetFlashImageColorAndOpacity() const;
+
+	static const FSlateBrush* GetBorderImage(const FObjectMixerEditorListRow::EObjectMixerEditorListRowType InRowType);
+
+	TSharedRef<SWidget> GenerateCells(const FName& InColumnName, const TSharedPtr<FObjectMixerEditorListRow> PinnedItem);
+
+	void OnPropertyChanged(const FProperty* Property, void* ContainerWithChangedProperty);
+	
+	TWeakPtr<FObjectMixerEditorListRow> Item;
+	
+	TSharedPtr<IToolTip> HoverToolTip;
+
+	TArray<TSharedPtr<SImage>> FlashImages;
+
+	TSet<FDelegateHandle> StructureChangeDelegateHandles;
+
+	TSharedPtr<SObjectMixerEditorListValueInput> ValueChildInputWidget;
+	
+	TSharedPtr<SObjectMixerEditorListRowHoverWidgets> HoverableWidgetsPtr;
+
+	FCurveSequence FlashAnimation;
+
+	const float FlashAnimationDuration = 0.75f;
+	const FLinearColor FlashColor = FLinearColor::White;
+	
+	const FSlateBrush* VisibleHoveredBrush = nullptr;
+	const FSlateBrush* VisibleNotHoveredBrush = nullptr;
+	const FSlateBrush* NotVisibleHoveredBrush = nullptr;
+	const FSlateBrush* NotVisibleNotHoveredBrush = nullptr;
+
+	/** The offset applied to text widgets so that the text aligns with the column header text */
+	float TextBlockLeftPadding = 3.0f;
+	
+	bool bIsHovered = false;
+};
diff --git a/Engine/Source/Editor/ObjectMixer/Public/Views/MainPanel/ObjectMixerEditorMainPanel.h b/Engine/Source/Editor/ObjectMixer/Public/Views/MainPanel/ObjectMixerEditorMainPanel.h
new file mode 100644
--- /dev/null
+++ b/Engine/Source/Editor/ObjectMixer/Public/Views/MainPanel/ObjectMixerEditorMainPanel.h
@@ -0,0 +1,46 @@
+// Copyright Epic Games, Inc. All Rights Reserved.
+
+#pragma once
+
+#include "Views/List/ObjectMixerEditorList.h"
+#include "Widgets/SWidget.h"
+
+class FObjectMixerEditorList;
+class SObjectMixerEditorMainPanel;
+
+class OBJECTMIXEREDITOR_API FObjectMixerEditorMainPanel : public TSharedFromThis<FObjectMixerEditorMainPanel>
+{
+public:
+	FObjectMixerEditorMainPanel();
+
+	~FObjectMixerEditorMainPanel() = default;
+
+	TSharedRef<SWidget> GetOrCreateWidget();
+
+	/*
+	 * Regenerate the list items and refresh the list. Call when adding or removing variables.
+	 * @param bShouldCacheValues If true, the current list's current values will be cached and then restored when the list is rebuilt. Otherwise preset values will be used.
+	 */
+	void RebuildList(const FString InItemToScrollTo = "", bool bShouldCacheValues = true) const;
+
+	/**
+	 * Refresh filters and sorting.
+	 * Useful for when the list state has gone stale but the variable count has not changed.
+	 */
+	void RefreshList() const;
+
+	TWeakPtr<FObjectMixerEditorList> GetEditorList() const
+	{
+		return EditorList;
+	}
+
+	void OnClassSelectionChanged(UClass* InNewClass) const;
+	TObjectPtr<UClass> GetClassSelection() const;
+	bool IsClassSelected(UClass* InNewClass) const;
+
+private:
+
+	TSharedPtr<SObjectMixerEditorMainPanel> MainPanelWidget;
+
+	TSharedPtr<FObjectMixerEditorList> EditorList;
+};
diff --git a/Engine/Source/Editor/ObjectMixer/Public/Views/MainPanel/SObjectMixerEditorMainPanel.h b/Engine/Source/Editor/ObjectMixer/Public/Views/MainPanel/SObjectMixerEditorMainPanel.h
new file mode 100644
--- /dev/null
+++ b/Engine/Source/Editor/ObjectMixer/Public/Views/MainPanel/SObjectMixerEditorMainPanel.h
@@ -0,0 +1,41 @@
+// Copyright Epic Games, Inc. All Rights Reserved.
+
+#pragma once
+
+#include "Widgets/Input/SCheckBox.h"
+#include "Widgets/SCompoundWidget.h"
+
+class FObjectMixerEditorMainPanel;
+class SBox;
+class SHorizontalBox;
+class SVerticalBox;
+
+class OBJECTMIXEREDITOR_API SObjectMixerEditorMainPanel final : public SCompoundWidget
+{
+
+public:
+
+	SLATE_BEGIN_ARGS(SObjectMixerEditorMainPanel)
+	{}
+
+	SLATE_END_ARGS()
+
+	void Construct(const FArguments& InArgs, const TSharedRef<FObjectMixerEditorMainPanel>& InMainPanel);
+
+	virtual ~SObjectMixerEditorMainPanel() override;
+
+private:
+
+	/** A reference to the struct that controls this widget */
+	TWeakPtr<FObjectMixerEditorMainPanel> MainPanel;
+
+	TSharedPtr<SHorizontalBox> ToolbarHBox;
+
+	/** A reference to the button which opens the plugin settings */
+	TSharedPtr<SCheckBox> ConcertButtonPtr;
+
+	/** Creates the toolbar at the top of the MainPanel widget */
+	TSharedRef<SWidget> GeneratePanelToolbar();
+
+	TSharedRef<SWidget> OnGeneratePresetsMenu();
+};
diff --git a/Engine/Source/Editor/ObjectMixer/Public/Views/Widgets/SObjectMixerEditorTooltipWidget.h b/Engine/Source/Editor/ObjectMixer/Public/Views/Widgets/SObjectMixerEditorTooltipWidget.h
new file mode 100644
--- /dev/null
+++ b/Engine/Source/Editor/ObjectMixer/Public/Views/Widgets/SObjectMixerEditorTooltipWidget.h
@@ -0,0 +1,79 @@
+// Copyright Epic Games, Inc. All Rights Reserved.
+
+#pragma once
+
+#include "Styling/AppStyle.h"
+#include "Widgets/Layout/SBorder.h"
+#include "Widgets/SBoxPanel.h"
+#include "Widgets/SCompoundWidget.h"
+#include "Widgets/SToolTip.h"
+#include "Widgets/Text/STextBlock.h"
+
+class OBJECTMIXEREDITOR_API SObjectMixerEditorTooltipWidget : public SCompoundWidget
+{
+public:
+	
+	SLATE_BEGIN_ARGS(SObjectMixerEditorTooltipWidget)
+	{}
+
+	SLATE_END_ARGS()
+	
+	void Construct(const FArguments& InArgs, const FString InCommandName, const FString InHelpString)
+	{
+		TSharedPtr<SVerticalBox> VerticalBox;
+		
+		ChildSlot
+		[
+			SAssignNew(VerticalBox, SVerticalBox)
+
+			+SVerticalBox::Slot()
+			.AutoHeight()
+			.Padding(0, 0, 0, 4)
+			[
+				SNew(SBorder)
+				.Padding(6)
+				.BorderImage(FAppStyle::Get().GetBrush("ContentBrowser.TileViewTooltip.ContentBorder"))
+				[
+					SNew(STextBlock)
+					.Text(FText::FromString(InCommandName))
+					.Font(FAppStyle::Get().GetFontStyle("ContentBrowser.TileViewTooltip.NameFont"))
+					.AutoWrapText(true)
+				]
+			]
+		];
+
+		if (!InHelpString.IsEmpty())
+		{
+			VerticalBox->AddSlot()
+			.AutoHeight()
+			[
+				SNew(SBorder)
+				.Padding(6)
+				.BorderImage(FAppStyle::Get().GetBrush("ContentBrowser.TileViewTooltip.ContentBorder"))
+				[
+					SNew(STextBlock)
+					.Text(FText::FromString(InHelpString))
+					.AutoWrapText(true)
+				]
+			];
+		}
+	}
+
+	static TSharedRef<SToolTip> MakeTooltip(const FString InCommandName, const FString InHelpString)
+	{
+		return
+			SNew(SToolTip)
+			.BorderImage(FAppStyle::Get().GetBrush("ContentBrowser.TileViewTooltip.NonContentBorder"))
+			[
+				SNew(SBox)
+				.WidthOverride(TooltipWindowWidth)
+				[
+					SNew(SObjectMixerEditorTooltipWidget, InCommandName, InHelpString)
+				]
+			];
+	}
+
+private:
+
+	static constexpr float TooltipWindowWidth = 500.f;
+};
