Fix for: //UE5/Main - error: #pragma once in main file in ObjectMixerEditorObjectFilter.cpp

diff --git a/Engine/Source/Editor/ObjectMixer/Private/ObjectFilter/ObjectMixerEditorObjectFilter.cpp b/Engine/Source/Editor/ObjectMixer/Private/ObjectFilter/ObjectMixerEditorObjectFilter.cpp
--- a/Engine/Source/Editor/ObjectMixer/Private/ObjectFilter/ObjectMixerEditorObjectFilter.cpp
+++ b/Engine/Source/Editor/ObjectMixer/Private/ObjectFilter/ObjectMixerEditorObjectFilter.cpp
@@ -1,7 +1,5 @@
 // Copyright Epic Games, Inc. All Rights Reserved.
 
-#pragma once
-
 #include "ObjectFilter/ObjectMixerEditorObjectFilter.h"
 
 #include "GameFramework/Actor.h"
