//UE5/Main - Run Enterprise UnrealEditor Python Tests - from_dataprep_to_unreal tests failing due to PropertyType is not initialized properly in SObjectMixerEditorList.h

diff --git a/Engine/Source/Editor/ObjectMixer/ObjectMixer/Public/Views/List/SObjectMixerEditorList.h b/Engine/Source/Editor/ObjectMixer/ObjectMixer/Public/Views/List/SObjectMixerEditorList.h
--- a/Engine/Source/Editor/ObjectMixer/ObjectMixer/Public/Views/List/SObjectMixerEditorList.h
+++ b/Engine/Source/Editor/ObjectMixer/ObjectMixer/Public/Views/List/SObjectMixerEditorList.h
@@ -33,19 +33,19 @@ struct FListViewColumnInfo
 {
 	GENERATED_BODY()
 	
-	TObjectPtr<FProperty> PropertyRef;
+	TObjectPtr<FProperty> PropertyRef = nullptr;
 	
 	UPROPERTY()
-	FName PropertyName;
+	FName PropertyName = NAME_None;
 
 	UPROPERTY()
-	FText PropertyDisplayText;
+	FText PropertyDisplayText = FText::GetEmpty();
 
 	UPROPERTY()
-	EListViewColumnType PropertyType;
+	EListViewColumnType PropertyType = EListViewColumnType::PropertyGenerated;
 
 	UPROPERTY()
-	FName CategoryName;
+	FName CategoryName = NAME_None;
 
 	UPROPERTY()
 	bool bIsDesiredForDisplay = false;
