Object Mixer/ Light Mixer

-Support Folder and attach parent structure
-Support placing actors from OM panel
-Make Blueprint rows have a clickable link to edit the blueprint
-Support Visibility and Mute controls for folders in OM panel
-Support multiple view modes
-Save/restore tree expansion state when tree is refreshed
-Cache filter data for performance reasons
-Expose Tree View mode to filter class overrides
-Convert Filter arrays to sets
-Move Class Filter Selection combobox into Show Options (cog)
-Sort tree items by row type then alphabetically
-Performance Optimization Pass
-Fix comment blocks from /* to /** for proper parsing

diff --git a/Engine/Source/Editor/LevelEditor/Private/SLevelEditor.cpp b/Engine/Source/Editor/LevelEditor/Private/SLevelEditor.cpp
--- a/Engine/Source/Editor/LevelEditor/Private/SLevelEditor.cpp
+++ b/Engine/Source/Editor/LevelEditor/Private/SLevelEditor.cpp
@@ -1488,13 +1488,19 @@ TSharedRef<SWidget> SLevelEditor::RestoreContentArea( const TSharedRef<SDockTab>
 			FObjectMixerEditorModule& ObjectMixer = FObjectMixerEditorModule::Get();
 			{
 				FDelegateHandle ObjectMixerComponentEditedHandle =
-				   LevelEditorModule.OnComponentsEdited().AddRaw(&ObjectMixer, &FObjectMixerEditorModule::RequestRebuildList);
+					LevelEditorModule.OnComponentsEdited().AddLambda([]()
+					{
+						FObjectMixerEditorModule::Get().RequestRebuildList();
+					});
 				ObjectMixer.AddOnComponentEditedDelegate(ObjectMixerComponentEditedHandle);
 			}
 			FLightMixerModule& LightMixer = FLightMixerModule::Get();
 			{
 				FDelegateHandle LightMixerComponentEditedHandle =
-				   LevelEditorModule.OnComponentsEdited().AddRaw(&LightMixer, &FLightMixerModule::RequestRebuildList);
+				   LevelEditorModule.OnComponentsEdited().AddLambda([]()
+				   {
+					   FLightMixerModule::Get().RequestRebuildList();
+				   });
 				LightMixer.AddOnComponentEditedDelegate(LightMixerComponentEditedHandle);
 			}
 		}
diff --git a/Engine/Source/Editor/ObjectMixer/LightMixer/Public/LightMixerObjectFilter.h b/Engine/Source/Editor/ObjectMixer/LightMixer/Public/LightMixerObjectFilter.h
--- a/Engine/Source/Editor/ObjectMixer/LightMixer/Public/LightMixerObjectFilter.h
+++ b/Engine/Source/Editor/ObjectMixer/LightMixer/Public/LightMixerObjectFilter.h
@@ -4,6 +4,7 @@
 
 #include "ObjectFilter/ObjectMixerEditorObjectFilter.h"
 
+#include "Engine/Light.h"
 #include "GameFramework/Actor.h"
 
 #include "LightMixerObjectFilter.generated.h"
@@ -14,7 +15,7 @@ class LIGHTMIXER_API ULightMixerObjectFilter : public UObjectMixerObjectFilter
 	GENERATED_BODY()
 public:
 	
-	virtual TArray<UClass*> GetObjectClassesToFilter() const override
+	virtual TSet<UClass*> GetObjectClassesToFilter() const override
 	{
 		return
 		{
@@ -22,30 +23,23 @@ public:
 		};
 	}
 
-	virtual FText GetRowDisplayName(UObject* InObject) const override
+	virtual TSet<TSubclassOf<AActor>> GetObjectClassesToPlace() const override
 	{
-		if (InObject)
+		return
 		{
-			if (const AActor* Outer = InObject->GetTypedOuter<AActor>())
-			{
-				return FText::Format(INVTEXT("{0} ({1})"), FText::FromString(Outer->GetActorLabel()), Super::GetRowDisplayName(InObject));
-			}
-		}
-
-		return Super::GetRowDisplayName(InObject);
+			ALight::StaticClass()
+		};
 	}
 
-	virtual TArray<FName> GetColumnsToShowByDefault() const override
+	virtual TSet<FName> GetColumnsToShowByDefault() const override
 	{
-		const TArray<FName> IncludeList =
+		return 
 		{
 			"Intensity", "LightColor", "AttenuationRadius"
 		};
-
-		return IncludeList;
 	}
 
-	virtual TArray<FName> GetForceAddedColumns() const override
+	virtual TSet<FName> GetForceAddedColumns() const override
 	{
 		return {"LightColor"};
 	}
@@ -55,8 +49,13 @@ public:
 		return false;
 	}
 
-	virtual EObjectMixerPropertyInheritanceInclusionOptions GetObjectMixerPropertyInheritanceInclusionOptions() const override
+	virtual EObjectMixerInheritanceInclusionOptions GetObjectMixerPropertyInheritanceInclusionOptions() const override
+	{
+		return EObjectMixerInheritanceInclusionOptions::IncludeAllParentsAndChildren;
+	}
+
+	virtual EObjectMixerInheritanceInclusionOptions GetObjectMixerPlacementClassInclusionOptions() const override
 	{
-		return EObjectMixerPropertyInheritanceInclusionOptions::IncludeAllParentsAndChildren;
+		return EObjectMixerInheritanceInclusionOptions::IncludeAllChildren;
 	}
 };
diff --git a/Engine/Source/Editor/ObjectMixer/LightMixer/Public/LightMixerProjectSettings.h b/Engine/Source/Editor/ObjectMixer/LightMixer/Public/LightMixerProjectSettings.h
--- a/Engine/Source/Editor/ObjectMixer/LightMixer/Public/LightMixerProjectSettings.h
+++ b/Engine/Source/Editor/ObjectMixer/LightMixer/Public/LightMixerProjectSettings.h
@@ -15,7 +15,7 @@ public:
 
 	/**
 	 * If true, the Object Mixer menu item will be disabled and removed.
-	 * This is useful if you only want to use Light Mixer and avoid cluttering the menus.
+	 * This is useful if you only want to use Light Mixer and to avoid cluttering the menus.
 	 */
 	UPROPERTY(Config, EditAnywhere, Category="Light Mixer", meta = (ConfigRestartRequired = true))
 	bool bHideObjectMixerMenuItem = false;
diff --git a/Engine/Source/Editor/ObjectMixer/ObjectMixer/Private/ObjectFilter/ObjectMixerEditorObjectFilter.cpp b/Engine/Source/Editor/ObjectMixer/ObjectMixer/Private/ObjectFilter/ObjectMixerEditorObjectFilter.cpp
--- a/Engine/Source/Editor/ObjectMixer/ObjectMixer/Private/ObjectFilter/ObjectMixerEditorObjectFilter.cpp
+++ b/Engine/Source/Editor/ObjectMixer/ObjectMixer/Private/ObjectFilter/ObjectMixerEditorObjectFilter.cpp
@@ -3,20 +3,51 @@
 #include "ObjectFilter/ObjectMixerEditorObjectFilter.h"
 
 #include "GameFramework/Actor.h"
+#include "Kismet2/ComponentEditorUtils.h"
 
-FText UObjectMixerObjectFilter::GetRowDisplayName(UObject* InObject) const
+FText UObjectMixerObjectFilter::GetRowDisplayName(UObject* InObject, EObjectMixerTreeViewMode InViewMode) const
 {
-	if (InObject->IsValidLowLevel())
+	if (IsValid(InObject))
 	{
+		if (InViewMode == EObjectMixerTreeViewMode::Flat || InViewMode == EObjectMixerTreeViewMode::Folder)
+		{
+			if (const AActor* Outer = InObject->GetTypedOuter<AActor>())
+			{
+				return FText::Format(INVTEXT("{0} ({1})"), FText::FromString(Outer->GetActorLabel()), FText::FromString(InObject->GetName()));
+			}
+		}
+
+		if (const AActor* AsActor = Cast<AActor>(InObject))
+		{
+			return FText::FromString(AsActor->GetActorLabel());
+		}
+		else if (const UActorComponent* AsActorComponent = Cast<UActorComponent>(InObject))
+		{
+			const FName ComponentName = FComponentEditorUtils::FindVariableNameGivenComponentInstance(AsActorComponent);
+			if (ComponentName != NAME_None)
+			{
+				if (!ComponentName.IsEqual(AsActorComponent->GetFName()))
+				{
+					return
+					   FText::Format(INVTEXT("{0} ({1})"),
+						   FText::FromName(ComponentName),
+						   FText::FromName(AsActorComponent->GetFName())
+					   );
+				}
+
+				return FText::FromName(ComponentName);
+			}
+		}
+		
 		return FText::FromString(InObject->GetName());
 	}
 
 	return FText::GetEmpty();
 }
 
-bool UObjectMixerObjectFilter::GetRowEditorVisibility(UObject* InObject) const
+bool UObjectMixerObjectFilter::GetRowEditorVisibility(UObject* InObject, EObjectMixerTreeViewMode InViewMode) const
 {
-	if (InObject->IsValidLowLevel())
+	if (IsValid(InObject))
 	{
 		TObjectPtr<AActor> Actor = Cast<AActor>(InObject);
 	
@@ -31,9 +62,9 @@ bool UObjectMixerObjectFilter::GetRowEditorVisibility(UObject* InObject) const
 	return false;
 }
 
-void UObjectMixerObjectFilter::OnSetRowEditorVisibility(UObject* InObject, bool bNewIsVisible) const
+void UObjectMixerObjectFilter::OnSetRowEditorVisibility(UObject* InObject, bool bNewIsVisible, EObjectMixerTreeViewMode InViewMode) const
 {
-	if (InObject->IsValidLowLevel())
+	if (IsValid(InObject))
 	{
 		TObjectPtr<AActor> Actor = Cast<AActor>(InObject);
 	
@@ -49,19 +80,24 @@ void UObjectMixerObjectFilter::OnSetRowEditorVisibility(UObject* InObject, bool
 	}
 }
 
-TArray<FName> UObjectMixerObjectFilter::GetColumnsToShowByDefault() const
+TSet<FName> UObjectMixerObjectFilter::GetColumnsToShowByDefault() const
 {
 	return {};
 }
 
-TArray<FName> UObjectMixerObjectFilter::GetColumnsFilter() const
+TSet<FName> UObjectMixerObjectFilter::GetColumnsToExclude() const
 {
 	return {};
 }
 
-EObjectMixerPropertyInheritanceInclusionOptions UObjectMixerObjectFilter::GetObjectMixerPropertyInheritanceInclusionOptions() const
+EObjectMixerInheritanceInclusionOptions UObjectMixerObjectFilter::GetObjectMixerPropertyInheritanceInclusionOptions() const
+{
+	return EObjectMixerInheritanceInclusionOptions::None;
+}
+
+EObjectMixerInheritanceInclusionOptions UObjectMixerObjectFilter::GetObjectMixerPlacementClassInclusionOptions() const
 {
-	return EObjectMixerPropertyInheritanceInclusionOptions::None;
+	return EObjectMixerInheritanceInclusionOptions::None;
 }
 
 bool UObjectMixerObjectFilter::ShouldIncludeUnsupportedProperties() const
@@ -69,18 +105,18 @@ bool UObjectMixerObjectFilter::ShouldIncludeUnsupportedProperties() const
 	return false;
 }
 
-TArray<UClass*> UObjectMixerObjectFilter::GetParentAndChildClassesFromSpecifiedClasses(
-	const TArray<UClass*>& InSpecifiedClasses, EObjectMixerPropertyInheritanceInclusionOptions Options)
+TSet<UClass*> UObjectMixerObjectFilter::GetParentAndChildClassesFromSpecifiedClasses(
+	const TSet<UClass*>& InSpecifiedClasses, EObjectMixerInheritanceInclusionOptions Options)
 {
 	// 'None' means we only want the specified classes
-	if (Options == EObjectMixerPropertyInheritanceInclusionOptions::None)
+	if (Options == EObjectMixerInheritanceInclusionOptions::None)
 	{
 		return InSpecifiedClasses;
 	}
 	
-	TArray<UClass*> ReturnValue;
+	TSet<UClass*> ReturnValue;
 
-	auto GetChildClassesLambda = [](const UClass* Class, TArray<UClass*>& OutReturnValue, const bool bRecursive)
+	auto GetChildClassesLambda = [](const UClass* Class, TSet<UClass*>& OutReturnValue, const bool bRecursive)
 	{
 		TArray<UClass*> DerivedClasses;
 		GetDerivedClasses(Class, DerivedClasses, bRecursive);
@@ -94,9 +130,9 @@ TArray<UClass*> UObjectMixerObjectFilter::GetParentAndChildClassesFromSpecifiedC
 		// Super Classes
 
 			// Immediate only
-		if (Options == EObjectMixerPropertyInheritanceInclusionOptions::IncludeOnlyImmediateParent ||
-			Options == EObjectMixerPropertyInheritanceInclusionOptions::IncludeOnlyImmediateParentAndChildren ||
-			Options == EObjectMixerPropertyInheritanceInclusionOptions::IncludeOnlyImmediateParentAndAllChildren)
+		if (Options == EObjectMixerInheritanceInclusionOptions::IncludeOnlyImmediateParent ||
+			Options == EObjectMixerInheritanceInclusionOptions::IncludeOnlyImmediateParentAndChildren ||
+			Options == EObjectMixerInheritanceInclusionOptions::IncludeOnlyImmediateParentAndAllChildren)
 		{
 			if (UClass* Super = Class->GetSuperClass())
 			{
@@ -104,9 +140,9 @@ TArray<UClass*> UObjectMixerObjectFilter::GetParentAndChildClassesFromSpecifiedC
 			}
 		}
 			// All Parents
-		if (Options == EObjectMixerPropertyInheritanceInclusionOptions::IncludeAllParents ||
-			Options == EObjectMixerPropertyInheritanceInclusionOptions::IncludeAllParentsAndChildren ||
-			Options == EObjectMixerPropertyInheritanceInclusionOptions::IncludeAllParentsAndOnlyImmediateChildren)
+		if (Options == EObjectMixerInheritanceInclusionOptions::IncludeAllParents ||
+			Options == EObjectMixerInheritanceInclusionOptions::IncludeAllParentsAndChildren ||
+			Options == EObjectMixerInheritanceInclusionOptions::IncludeAllParentsAndOnlyImmediateChildren)
 		{
 			if (UClass* Super = Class->GetSuperClass())
 			{
@@ -121,16 +157,16 @@ TArray<UClass*> UObjectMixerObjectFilter::GetParentAndChildClassesFromSpecifiedC
 		// Child Classes
 
 			// Immediate only
-		if (Options == EObjectMixerPropertyInheritanceInclusionOptions::IncludeOnlyImmediateChildren ||
-			Options == EObjectMixerPropertyInheritanceInclusionOptions::IncludeOnlyImmediateParentAndChildren ||
-			Options == EObjectMixerPropertyInheritanceInclusionOptions::IncludeAllParentsAndOnlyImmediateChildren)
+		if (Options == EObjectMixerInheritanceInclusionOptions::IncludeOnlyImmediateChildren ||
+			Options == EObjectMixerInheritanceInclusionOptions::IncludeOnlyImmediateParentAndChildren ||
+			Options == EObjectMixerInheritanceInclusionOptions::IncludeAllParentsAndOnlyImmediateChildren)
 		{
 			GetChildClassesLambda(Class, ReturnValue, false);
 		}
 			// All Children
-		if (Options == EObjectMixerPropertyInheritanceInclusionOptions::IncludeAllChildren ||
-			Options == EObjectMixerPropertyInheritanceInclusionOptions::IncludeAllParentsAndChildren ||
-			Options == EObjectMixerPropertyInheritanceInclusionOptions::IncludeOnlyImmediateParentAndAllChildren)
+		if (Options == EObjectMixerInheritanceInclusionOptions::IncludeAllChildren ||
+			Options == EObjectMixerInheritanceInclusionOptions::IncludeAllParentsAndChildren ||
+			Options == EObjectMixerInheritanceInclusionOptions::IncludeOnlyImmediateParentAndAllChildren)
 		{
 			GetChildClassesLambda(Class, ReturnValue, true);
 		}
@@ -139,18 +175,31 @@ TArray<UClass*> UObjectMixerObjectFilter::GetParentAndChildClassesFromSpecifiedC
 	return ReturnValue;
 }
 
+TSet<UClass*> UObjectMixerObjectFilter::GetParentAndChildClassesFromSpecifiedClasses(
+	const TSet<TSubclassOf<AActor>>& InSpecifiedClasses, EObjectMixerInheritanceInclusionOptions Options)
+{
+	TSet<UClass*> AsUClasses;
+	
+	for (TSubclassOf<AActor> ActorClass : InSpecifiedClasses)
+	{
+		AsUClasses.Add(ActorClass);
+	}
+
+	return GetParentAndChildClassesFromSpecifiedClasses(AsUClasses, Options);
+}
+
 EFieldIterationFlags UObjectMixerObjectFilter::GetDesiredFieldIterationFlags(const bool bIncludeInheritedProperties)
 {
 	return bIncludeInheritedProperties ? EFieldIterationFlags::IncludeSuper : EFieldIterationFlags::Default;
 }
 
-TArray<FName> UObjectMixerObjectFilter::GenerateIncludeListFromExcludeList(const TSet<FName>& ExcludeList) const
+TSet<FName> UObjectMixerObjectFilter::GenerateIncludeListFromExcludeList(const TSet<FName>& ExcludeList) const
 {
 	TSet<FName> IncludeList;
 
-	const EObjectMixerPropertyInheritanceInclusionOptions Options =
+	const EObjectMixerInheritanceInclusionOptions Options =
 		GetObjectMixerPropertyInheritanceInclusionOptions();
-	TArray<UClass*> SpecifiedClasses =
+	TSet<UClass*> SpecifiedClasses =
 		GetParentAndChildClassesFromSpecifiedClasses(GetObjectClassesToFilter(), Options);
 	
 	for (const UClass* Class : SpecifiedClasses)
@@ -164,5 +213,5 @@ TArray<FName> UObjectMixerObjectFilter::GenerateIncludeListFromExcludeList(const
 		}
 	}
 
-	return IncludeList.Difference(ExcludeList).Array();
+	return IncludeList.Difference(ExcludeList);
 }
diff --git a/Engine/Source/Editor/ObjectMixer/ObjectMixer/Private/ObjectMixerEditorModule.cpp b/Engine/Source/Editor/ObjectMixer/ObjectMixer/Private/ObjectMixerEditorModule.cpp
--- a/Engine/Source/Editor/ObjectMixer/ObjectMixer/Private/ObjectMixerEditorModule.cpp
+++ b/Engine/Source/Editor/ObjectMixer/ObjectMixer/Private/ObjectMixerEditorModule.cpp
@@ -40,10 +40,7 @@ void FObjectMixerEditorModule::ShutdownModule()
 
 void FObjectMixerEditorModule::Initialize()
 {
-	DelegateHandles.Add(FEditorDelegates::MapChange.AddLambda([this](uint32 Index)
-	{
-		RequestRebuildList();
-	}));
+	BindDelegates();
 	
 	SetupMenuItemVariables();
 	
@@ -54,7 +51,25 @@ void FObjectMixerEditorModule::Initialize()
 void FObjectMixerEditorModule::Teardown()
 {
 	// Unbind Delegates
+	if (GEngine)
+	{
+		GEngine->OnLevelActorFolderChanged().RemoveAll(this);
+		GEngine->OnLevelActorListChanged().RemoveAll(this);
+		GEngine->OnActorFolderAdded().RemoveAll(this);
+		GEngine->OnLevelActorDetached().RemoveAll(this);
+		GEngine->OnLevelActorAttached().RemoveAll(this);
+		GEngine->OnActorFolderAdded().RemoveAll(this);
+		GEngine->OnActorFoldersUpdatedEvent().RemoveAll(this);
+	}
+
+	if (GEditor)
+	{
+		GEditor->OnBlueprintCompiled().RemoveAll(this);
+	}
+	
 	FEditorDelegates::MapChange.RemoveAll(this);
+	FEditorDelegates::PostUndoRedo.RemoveAll(this);
+	FCoreUObjectDelegates::OnObjectTransacted.RemoveAll(this);
 
 	for (FDelegateHandle& Delegate : DelegateHandles)
 	{
@@ -208,6 +223,11 @@ TSharedRef<SDockTab> FObjectMixerEditorModule::SpawnTab(const FSpawnTabArgs& Arg
 	 return SpawnMainPanelTab();
 }
 
+TSharedPtr<FWorkspaceItem> FObjectMixerEditorModule::GetWorkspaceGroup()
+{
+	return WorkspaceGroup;
+}
+
 TSharedRef<SDockTab> FObjectMixerEditorModule::SpawnMainPanelTab()
 {
 	MainPanel = MakeShared<FObjectMixerEditorMainPanel>();
@@ -219,15 +239,69 @@ TSharedRef<SDockTab> FObjectMixerEditorModule::SpawnMainPanelTab()
 	;
 	const TSharedPtr<SWidget> ObjectMixerDialog = MakeObjectMixerDialog();
 	DockTab->SetContent(ObjectMixerDialog ? ObjectMixerDialog.ToSharedRef() : SNullWidget::NullWidget);
-	MainPanel->OnClassSelectionChanged(DefaultFilterClass);
-	MainPanel->RequestRebuildList();
+
+	if (DefaultFilterClass)
+	{
+		MainPanel->OnClassSelectionChanged(DefaultFilterClass);
+	}
 			
 	return DockTab;
 }
 
-TSharedPtr<FWorkspaceItem> FObjectMixerEditorModule::GetWorkspaceGroup()
+void FObjectMixerEditorModule::BindDelegates()
 {
-	return WorkspaceGroup;
+	check(GEngine && GEditor);
+	DelegateHandles.Add(GEngine->OnLevelActorFolderChanged().AddLambda([this] (const AActor*, FName)
+	{
+		RequestRebuildList();
+	}));
+	DelegateHandles.Add(GEngine->OnLevelActorListChanged().AddLambda([this] ()
+	{
+		RequestRebuildList();
+	}));
+	DelegateHandles.Add(GEngine->OnActorFolderAdded().AddLambda([this] (UActorFolder*)
+	{
+		RequestRebuildList();
+	}));
+	DelegateHandles.Add(GEngine->OnLevelActorDetached().AddLambda([this] (AActor*, const AActor*)
+	{
+		RequestRebuildList();
+	}));
+	DelegateHandles.Add(GEngine->OnLevelActorAttached().AddLambda([this] (AActor*, const AActor*)
+	{
+		RequestRebuildList();
+	}));
+	DelegateHandles.Add(GEngine->OnActorFolderAdded().AddLambda([this] (UActorFolder*)
+	{
+		RequestRebuildList();
+	}));
+	DelegateHandles.Add(GEngine->OnActorFoldersUpdatedEvent().AddLambda([this] (ULevel*)
+	{
+		RequestRebuildList();
+	}));
+
+	DelegateHandles.Add(GEditor->OnBlueprintCompiled().AddLambda([this] ()
+	{
+		RequestRebuildList();
+	}));
+
+	DelegateHandles.Add(FEditorDelegates::MapChange.AddLambda([this](uint32 Index)
+	{
+		RequestRebuildList();
+	}));
+
+	DelegateHandles.Add(FEditorDelegates::PostUndoRedo.AddLambda([this]()
+	{
+		RequestRebuildList();
+	}));
+
+	DelegateHandles.Add(FCoreUObjectDelegates::OnObjectTransacted.AddLambda([this](UObject*, const FTransactionObjectEvent& Event)
+	{
+		if (Event.GetEventType() == ETransactionObjectEventType::Finalized)
+		{
+			RequestRebuildList();
+		}
+	}));
 }
 
 #undef LOCTEXT_NAMESPACE
diff --git a/Engine/Source/Editor/ObjectMixer/ObjectMixer/Private/Views/List/ObjectMixerEditorList.cpp b/Engine/Source/Editor/ObjectMixer/ObjectMixer/Private/Views/List/ObjectMixerEditorList.cpp
--- a/Engine/Source/Editor/ObjectMixer/ObjectMixer/Private/Views/List/ObjectMixerEditorList.cpp
+++ b/Engine/Source/Editor/ObjectMixer/ObjectMixer/Private/Views/List/ObjectMixerEditorList.cpp
@@ -84,3 +84,8 @@ void FObjectMixerEditorList::RefreshList() const
 		ListWidget->RefreshList();
 	}
 }
+
+bool FObjectMixerEditorList::IsClassSelected(UClass* InNewClass) const
+{
+	return InNewClass == GetObjectFilterClass();
+}
diff --git a/Engine/Source/Editor/ObjectMixer/ObjectMixer/Private/Views/List/ObjectMixerEditorListRow.cpp b/Engine/Source/Editor/ObjectMixer/ObjectMixer/Private/Views/List/ObjectMixerEditorListRow.cpp
--- a/Engine/Source/Editor/ObjectMixer/ObjectMixer/Private/Views/List/ObjectMixerEditorListRow.cpp
+++ b/Engine/Source/Editor/ObjectMixer/ObjectMixer/Private/Views/List/ObjectMixerEditorListRow.cpp
@@ -91,7 +91,9 @@ void FObjectMixerEditorListRow::SetChildRows(const TArray<FObjectMixerEditorList
 
 void FObjectMixerEditorListRow::AddToChildRows(const FObjectMixerEditorListRowPtr& InRow)
 {
-	ChildRows.Add(InRow);
+	InRow->SetDirectParentRow(SharedThis(this));
+	ChildRows.AddUnique(InRow);
+	ChildRows.StableSort(SObjectMixerEditorList::SortByTypeThenName);
 }
 
 void FObjectMixerEditorListRow::InsertChildRowAtIndex(const FObjectMixerEditorListRowPtr& InRow,
@@ -100,24 +102,14 @@ void FObjectMixerEditorListRow::InsertChildRowAtIndex(const FObjectMixerEditorLi
 	ChildRows.Insert(InRow, AtIndex);
 }
 
-bool FObjectMixerEditorListRow::GetIsTreeViewItemExpanded() const
+bool FObjectMixerEditorListRow::GetIsTreeViewItemExpanded()
 {
-	return bIsTreeViewItemExpanded;
+	return GetListViewPtr().Pin()->IsTreeViewItemExpanded(GetAsShared());
 }
 
 void FObjectMixerEditorListRow::SetIsTreeViewItemExpanded(const bool bNewExpanded)
 {
-	bIsTreeViewItemExpanded = bNewExpanded;
-}
-
-bool FObjectMixerEditorListRow::GetShouldFlashOnScrollIntoView() const
-{
-	return bShouldFlashOnScrollIntoView;
-}
-
-void FObjectMixerEditorListRow::SetShouldFlashOnScrollIntoView(const bool bNewShouldFlashOnScrollIntoView)
-{
-	bShouldFlashOnScrollIntoView = bNewShouldFlashOnScrollIntoView;
+	GetListViewPtr().Pin()->SetTreeViewItemExpanded(GetAsShared(), bNewExpanded);
 }
 
 bool FObjectMixerEditorListRow::GetShouldExpandAllChildren() const
@@ -130,11 +122,6 @@ void FObjectMixerEditorListRow::SetShouldExpandAllChildren(const bool bNewShould
 	bShouldExpandAllChildren = bNewShouldExpandAllChildren;
 }
 
-void FObjectMixerEditorListRow::ResetToStartupValueAndSource() const
-{
-
-}
-
 bool FObjectMixerEditorListRow::MatchSearchTokensToSearchTerms(
 	const TArray<FString> InTokens, ESearchCase::Type InSearchCase)
 {
@@ -149,7 +136,7 @@ bool FObjectMixerEditorListRow::MatchSearchTokensToSearchTerms(
 
 			if (const UObjectMixerObjectFilter* Filter = GetObjectFilter())
 			{
-				CachedSearchTerms = Filter->GetRowDisplayName(Object).ToString();
+				CachedSearchTerms = Filter->GetRowDisplayName(Object, GetTreeViewMode()).ToString();
 			}
 		}
 
@@ -201,7 +188,7 @@ void FObjectMixerEditorListRow::ExecuteSearchOnChildNodes(const TArray<FString>&
 			continue;
 		}
 
-		if (ChildRow->GetRowType() == EObjectMixerEditorListRowType::Group)
+		if (ChildRow->GetRowType() == EObjectMixerEditorListRowType::Folder)
 		{
 			if (ChildRow->MatchSearchTokensToSearchTerms(Tokens))
 			{
@@ -238,14 +225,48 @@ bool FObjectMixerEditorListRow::GetIsSelected()
 	return ListViewPtr.Pin()->IsTreeViewItemSelected(SharedThis(this));
 }
 
-bool FObjectMixerEditorListRow::ShouldBeVisible() const
+bool FObjectMixerEditorListRow::ShouldRowWidgetBeVisible() const
 {
-	return (bDoesRowMatchSearchTerms && bDoesRowPassFilters) || HasVisibleChildren();
+	return (bDoesRowMatchSearchTerms && bDoesRowPassFilters) || HasVisibleChildRowWidgets();
 }
 
-EVisibility FObjectMixerEditorListRow::GetDesiredVisibility() const
+EVisibility FObjectMixerEditorListRow::GetDesiredRowWidgetVisibility() const
 {
-	return ShouldBeVisible() ? EVisibility::Visible : EVisibility::Collapsed;
+	return ShouldRowWidgetBeVisible() ? EVisibility::Visible : EVisibility::Collapsed;
+}
+
+FText FObjectMixerEditorListRow::GetDisplayName()
+{
+	if (!GetDisplayNameOverride().IsEmpty())
+	{
+		return GetDisplayNameOverride();
+	}
+
+	if (const UObjectMixerObjectFilter* Filter = GetObjectFilter())
+	{
+		if (const TObjectPtr<UObject> Object = GetObject())
+		{
+			return Filter->GetRowDisplayName(Object, GetTreeViewMode());
+		}
+	}
+
+	return FText::GetEmpty();
+}
+
+EObjectMixerTreeViewMode FObjectMixerEditorListRow::GetTreeViewMode()
+{
+	const TSharedPtr<SObjectMixerEditorList> PinnedListModel = GetListViewPtr().Pin();
+	check(PinnedListModel);
+	
+	return PinnedListModel->GetTreeViewMode();
+}
+
+void FObjectMixerEditorListRow::SetTreeViewMode(EObjectMixerTreeViewMode InViewMode)
+{
+	if (TSharedPtr<SObjectMixerEditorList> PinnedListModel = GetListViewPtr().Pin())
+	{
+		PinnedListModel->SetTreeViewMode(InViewMode);
+	}
 }
 
 TArray<FObjectMixerEditorListRowPtr> FObjectMixerEditorListRow::GetSelectedTreeViewItems() const
@@ -255,10 +276,20 @@ TArray<FObjectMixerEditorListRowPtr> FObjectMixerEditorListRow::GetSelectedTreeV
 
 const FSlateBrush* FObjectMixerEditorListRow::GetObjectIconBrush()
 {
-	if (GetRowType() != SingleItem)
+	if (GetRowType() == EObjectMixerEditorListRowType::None)
 	{
 		return nullptr;
 	}
+
+	if (GetRowType() == EObjectMixerEditorListRowType::Folder)
+	{
+		if (GetIsTreeViewItemExpanded() && ChildRows.Num())
+		{
+			return FAppStyle::Get().GetBrush(TEXT("SceneOutliner.FolderOpen"));
+		}
+		
+		return FAppStyle::Get().GetBrush(TEXT("SceneOutliner.FolderClosed"));
+	}
 	
 	if (UObject* RowObject = GetObject())
 	{
@@ -279,12 +310,7 @@ const FSlateBrush* FObjectMixerEditorListRow::GetObjectIconBrush()
 		}
 	
 		if (RowObject->IsA(UActorComponent::StaticClass()))
-		{
-			if (AActor* OuterActor = RowObject->GetTypedOuter<AActor>())
-			{
-				return GetIconForActor(OuterActor);
-			}
-		
+		{		
 			return FSlateIconFinder::FindIconBrushForClass(RowObject->GetClass(), TEXT("SCS.Component"));
 		}
 	}
@@ -294,19 +320,42 @@ const FSlateBrush* FObjectMixerEditorListRow::GetObjectIconBrush()
 
 bool FObjectMixerEditorListRow::GetObjectVisibility()
 {
+	if (GetRowType() == EObjectMixerEditorListRowType::Folder)
+	{
+		// If any child returns true, the folder returns true
+		for (const FObjectMixerEditorListRowPtr& Child : GetChildRows())
+		{
+			if (Child->GetObjectVisibility())
+			{
+				return true;
+			}
+		}
+
+		// The folder's visibility is only returned false if all children also return false
+		return false;
+	}
+	
 	if (const UObjectMixerObjectFilter* Filter = GetObjectFilter())
 	{
-		return Filter->GetRowEditorVisibility(GetObject());
+		return Filter->GetRowEditorVisibility(GetObject(), GetTreeViewMode());
 	}
 
 	return false;
 }
 
-void FObjectMixerEditorListRow::SetObjectVisibility(const bool bNewIsVisible)
+void FObjectMixerEditorListRow::SetObjectVisibility(const bool bNewIsVisible, const bool bIsRecursive)
 {
 	if (const UObjectMixerObjectFilter* Filter = GetObjectFilter())
 	{
-		Filter->OnSetRowEditorVisibility(GetObject(), bNewIsVisible);
+		Filter->OnSetRowEditorVisibility(GetObject(), bNewIsVisible, GetTreeViewMode());
+
+		if (bIsRecursive)
+		{
+			for (const FObjectMixerEditorListRowPtr& Child : GetChildRows())
+			{
+				Child->SetObjectVisibility(bNewIsVisible, true);
+			}
+		}
 	}
 }
 
@@ -324,3 +373,8 @@ void FObjectMixerEditorListRow::ClearSoloRow()
 {
 	GetListViewPtr().Pin()->ClearSoloRow();
 }
+
+FObjectMixerEditorListRowPtr FObjectMixerEditorListRow::GetAsShared()
+{
+	return SharedThis(this);
+}
diff --git a/Engine/Source/Editor/ObjectMixer/ObjectMixer/Private/Views/List/SObjectMixerEditorList.cpp b/Engine/Source/Editor/ObjectMixer/ObjectMixer/Private/Views/List/SObjectMixerEditorList.cpp
--- a/Engine/Source/Editor/ObjectMixer/ObjectMixer/Private/Views/List/SObjectMixerEditorList.cpp
+++ b/Engine/Source/Editor/ObjectMixer/ObjectMixer/Private/Views/List/SObjectMixerEditorList.cpp
@@ -2,24 +2,26 @@
 
 #include "Views/List/SObjectMixerEditorList.h"
 
-// Filters
+// View Filters
 #include "Views/List/ObjectMixerEditorListFilters/ObjectMixerEditorListFilter_Source.h"
 
+#include "ObjectFilter/ObjectMixerEditorObjectFilter.h"
 #include "ObjectMixerEditorLog.h"
 #include "ObjectMixerEditorModule.h"
+#include "ObjectMixerEditorProjectSettings.h"
 #include "ObjectMixerEditorStyle.h"
 #include "Views/List/SObjectMixerEditorListRow.h"
+#include "Views/Widgets/SObjectMixerPlacementAssetMenuEntry.h"
 
 #include "Algo/AllOf.h"
 #include "Algo/AnyOf.h"
-#include "Algo/Find.h"
-#include "Components/LightComponent.h"
 #include "Editor.h"
 #include "Framework/MultiBox/MultiBoxBuilder.h"
-#include "ObjectFilter/ObjectMixerEditorObjectFilter.h"
+#include "SPositiveActionButton.h"
+#include "DragAndDrop/AssetDragDropOp.h"
+#include "PlacementMode/Public/IPlacementModeModule.h"
 #include "Styling/StyleColors.h"
 #include "UObject/UObjectIterator.h"
-#include "Widgets/Input/SCheckBox.h"
 #include "Widgets/Input/SSearchBox.h"
 #include "Widgets/Layout/SScrollBox.h"
 #include "Widgets/Layout/SWidgetSwitcher.h"
@@ -54,39 +56,7 @@ void SObjectMixerEditorList::Construct(const FArguments& InArgs, TSharedRef<FObj
 		.AutoHeight()
 		.Padding(FMargin(8.f, 0.f, 8.f, 0.f))
 		[
-			SNew(SHorizontalBox)
-
-			+SHorizontalBox::Slot()
-			.HAlign(HAlign_Fill)
-			.VAlign(VAlign_Center)
-			.Padding(0.f, 1.f, 0.f, 1.f)
-			[
-				SAssignNew(ListSearchBoxPtr, SSearchBox)
-				.HintText(LOCTEXT("SearchHintText", "Search Scene Objects"))
-				.ToolTipText(LOCTEXT("ObjectMixerEditorList_TooltipText", "Search Scene Objects"))
-				.OnTextChanged_Raw(this, &SObjectMixerEditorList::OnListViewSearchTextChanged)
-			]
-
-			// Show Options
-			+SHorizontalBox::Slot()
-			.HAlign(HAlign_Right)
-			.VAlign(VAlign_Center)
-			.AutoWidth()
-			.Padding(8.f, 1.f, 0.f, 1.f)
-			[
-				SAssignNew( ViewOptionsComboButton, SComboButton )
-				.ContentPadding(4.f)
-				.ToolTipText(LOCTEXT("ShowOptions_Tooltip", "Show options to affect the visibility of items in the Object Mixer list"))
-				.ComboButtonStyle( FAppStyle::Get(), "SimpleComboButtonWithIcon" ) // Use the tool bar item style for this button
-				.OnGetMenuContent( this, &SObjectMixerEditorList::BuildShowOptionsMenu)
-				.HasDownArrow(false)
-				.ButtonContent()
-				[
-					SNew(SImage)
-					.ColorAndOpacity(FSlateColor::UseForeground())
-					.Image( FAppStyle::Get().GetBrush("Icons.Settings") )
-				]
-			]
+			GenerateToolbar()
 		]
 
 		+ SVerticalBox::Slot()
@@ -111,11 +81,11 @@ void SObjectMixerEditorList::Construct(const FArguments& InArgs, TSharedRef<FObj
 				{
 					if (GEditor && FSlateApplication::Get().GetModifierKeys().IsAltDown())
 					{
-						GEditor->SelectNone(false, true, true);
-
+						TArray<AActor*> ActorsToSelect;
 						for (const TSharedPtr<FObjectMixerEditorListRow>& SelectedRow : TreeViewPtr->GetSelectedItems())
 						{
-							if (SelectedRow->GetRowType() == FObjectMixerEditorListRow::SingleItem)
+							if (SelectedRow->GetRowType() == FObjectMixerEditorListRow::MatchingObject ||
+								SelectedRow->GetRowType() == FObjectMixerEditorListRow::ContainerObject)
 							{
 								AActor* Actor = Cast<AActor>(SelectedRow->GetObject());
 
@@ -126,10 +96,20 @@ void SObjectMixerEditorList::Construct(const FArguments& InArgs, TSharedRef<FObj
 
 								if (Actor)
 								{
-									GEditor->SelectActor( Actor, true, true, true );
+									ActorsToSelect.Add(Actor);
 								}
 							}
 						}
+
+						if (ActorsToSelect.Num())
+						{
+							GEditor->SelectNone(true, true, true);
+
+							for (AActor* Actor : ActorsToSelect)
+							{
+								GEditor->SelectActor( Actor, true, true, true );
+							}
+						}
 					}
 				})
 				.TreeItemsSource(&VisibleTreeViewObjects)
@@ -138,7 +118,7 @@ void SObjectMixerEditorList::Construct(const FArguments& InArgs, TSharedRef<FObj
 						check(Row.IsValid());
 					
 						return SNew(SObjectMixerEditorListRow, TreeViewPtr.ToSharedRef(), Row)
-								.Visibility_Raw(Row.Get(), &FObjectMixerEditorListRow::GetDesiredVisibility);
+								.Visibility_Raw(Row.Get(), &FObjectMixerEditorListRow::GetDesiredRowWidgetVisibility);
 					})
 				.OnGetChildren_Raw(this, &SObjectMixerEditorList::OnGetRowChildren)
 				.OnExpansionChanged_Raw(this, &SObjectMixerEditorList::OnRowChildExpansionChange, false)
@@ -165,20 +145,7 @@ void SObjectMixerEditorList::Construct(const FArguments& InArgs, TSharedRef<FObj
 }
 
 SObjectMixerEditorList::~SObjectMixerEditorList()
-{
-	// Unbind Delegates
-	if (GEditor)
-	{
-		if (const UWorld* EditorWorld = GEditor->GetEditorWorldContext().World())
-		{
-			EditorWorld->RemoveOnActorSpawnedHandler(OnActorSpawnedHandle);
-			EditorWorld->RemoveOnActorDestroyededHandler(OnActorDestroyedHandle);
-		}
-	}
-
-	OnActorSpawnedHandle.Reset();
-	OnActorDestroyedHandle.Reset();
-	
+{	
 	HeaderRow.Reset();
 	
 	ListSearchBoxPtr.Reset();
@@ -191,6 +158,97 @@ SObjectMixerEditorList::~SObjectMixerEditorList()
 	TreeViewPtr.Reset();
 }
 
+TSharedRef<SWidget> SObjectMixerEditorList::GenerateToolbar()
+{
+	check(ListModelPtr.IsValid());
+	
+	TSharedRef<SHorizontalBox> ToolbarBox = SNew(SHorizontalBox);
+
+	// Add object button
+	ToolbarBox->AddSlot()
+   .HAlign(HAlign_Left)
+   .VAlign(VAlign_Center)
+   .AutoWidth()
+   .Padding(FMargin(0, 4))
+   [
+	   SNew(SPositiveActionButton)
+	   .Text(LOCTEXT("AddObject", "Add"))
+	   .OnGetMenuContent(FOnGetContent::CreateRaw(this, &SObjectMixerEditorList::OnGenerateAddObjectButtonMenu))
+   ];
+
+	ToolbarBox->AddSlot()
+	.HAlign(HAlign_Fill)
+	.VAlign(VAlign_Center)
+	.Padding(0.f, 1.f, 0.f, 1.f)
+	[
+		SAssignNew(ListSearchBoxPtr, SSearchBox)
+		.HintText(LOCTEXT("SearchHintText", "Search Scene Objects"))
+		.ToolTipText(LOCTEXT("ObjectMixerEditorList_TooltipText", "Search Scene Objects"))
+		.OnTextChanged_Raw(this, &SObjectMixerEditorList::OnListViewSearchTextChanged)
+	];
+
+	// Show Options
+	ToolbarBox->AddSlot()
+	.HAlign(HAlign_Right)
+	.VAlign(VAlign_Center)
+	.AutoWidth()
+	.Padding(8.f, 1.f, 0.f, 1.f)
+	[
+		SAssignNew( ViewOptionsComboButton, SComboButton )
+		.ContentPadding(4.f)
+		.ToolTipText(LOCTEXT("ShowOptions_Tooltip", "Show options to affect the visibility of items in the Object Mixer list"))
+		.ComboButtonStyle( FAppStyle::Get(), "SimpleComboButtonWithIcon" ) // Use the tool bar item style for this button
+		.OnGetMenuContent( this, &SObjectMixerEditorList::BuildShowOptionsMenu)
+		.HasDownArrow(false)
+		.ButtonContent()
+		[
+			SNew(SImage)
+			.ColorAndOpacity(FSlateColor::UseForeground())
+			.Image( FAppStyle::Get().GetBrush("Icons.Settings") )
+		]
+	];
+
+	return ToolbarBox;
+}
+
+TSharedRef<SWidget> SObjectMixerEditorList::OnGenerateAddObjectButtonMenu() const
+{
+	check(ListModelPtr.IsValid());
+
+	const TSet<TSubclassOf<AActor>> SubclassesOfActor = ListModelPtr.Pin()->GetObjectClassesToPlace();
+
+	if (SubclassesOfActor.Num() > 0)
+	{
+		TSet<UClass*> ClassesToPlace =
+		   ListModelPtr.Pin()->GetObjectFilter()->GetParentAndChildClassesFromSpecifiedClasses(
+			   SubclassesOfActor,
+			   ListModelPtr.Pin()->GetObjectFilter()->GetObjectMixerPlacementClassInclusionOptions());
+	
+		FMenuBuilder AddObjectButtonMenuBuilder = FMenuBuilder(true, nullptr);
+
+		for (const UClass* Class : ClassesToPlace)
+		{
+			if (const UActorFactory* Factory = GEditor->FindActorFactoryForActorClass(Class))
+			{
+				AddObjectButtonMenuBuilder.AddWidget(
+					SNew(SObjectMixerPlacementAssetMenuEntry, MakeShareable(new FPlaceableItem(*Factory->GetClass()))), FText::GetEmpty());
+			}
+		}
+
+		return AddObjectButtonMenuBuilder.MakeWidget();
+	}
+
+	return
+		SNew(SBox)
+		.Padding(5)
+		[
+			SNew(STextBlock)
+				.Text(LOCTEXT("NoPlaceableActorsDefinedWarning", "Please define some placeable actors in the\nfilter class by overriding GetObjectClassesToPlace."))
+				.Font(FAppStyle::Get().GetFontStyle("NormalFontItalic"))
+		]
+	;
+}
+
 void SObjectMixerEditorList::RebuildList()
 {
 	bShouldRebuild = false;
@@ -231,11 +289,11 @@ int32 SObjectMixerEditorList::GetSelectedTreeViewItemCount() const
 	return TreeViewPtr->GetSelectedItems().Num();
 }
 
-void SObjectMixerEditorList::SetSelectedTreeViewItemActorsEditorVisible(const bool bNewIsVisible)
+void SObjectMixerEditorList::SetSelectedTreeViewItemActorsEditorVisible(const bool bNewIsVisible, const bool bIsRecursive)
 {
 	for (const TSharedPtr<FObjectMixerEditorListRow>& SelectedItem : TreeViewPtr->GetSelectedItems())
 	{
-		SelectedItem->SetObjectVisibility(bNewIsVisible);
+		SelectedItem->SetObjectVisibility(bNewIsVisible, bIsRecursive);
 	}
 }
 
@@ -306,7 +364,7 @@ bool SObjectMixerEditorList::DoesTreeViewHaveVisibleChildren() const
 	{
 		for (const TSharedPtr<FObjectMixerEditorListRow>& Header : TreeViewRootObjects)
 		{
-			const EVisibility HeaderVisibility = Header->GetDesiredVisibility();
+			const EVisibility HeaderVisibility = Header->GetDesiredRowWidgetVisibility();
 			
 			if (HeaderVisibility != EVisibility::Hidden && HeaderVisibility != EVisibility::Collapsed)
 			{
@@ -318,6 +376,16 @@ bool SObjectMixerEditorList::DoesTreeViewHaveVisibleChildren() const
 	return false;
 }
 
+bool SObjectMixerEditorList::IsTreeViewItemExpanded(const TSharedPtr<FObjectMixerEditorListRow>& Row) const
+{
+	if (TreeViewPtr.IsValid())
+	{
+		return TreeViewPtr->IsItemExpanded(Row);
+	}
+	
+	return false;
+}
+
 void SObjectMixerEditorList::SetTreeViewItemExpanded(const TSharedPtr<FObjectMixerEditorListRow>& RowToExpand, const bool bNewExpansion) const
 {
 	if (TreeViewPtr.IsValid())
@@ -326,6 +394,22 @@ void SObjectMixerEditorList::SetTreeViewItemExpanded(const TSharedPtr<FObjectMix
 	}
 }
 
+EObjectMixerTreeViewMode SObjectMixerEditorList::GetTreeViewMode()
+{
+	const TSharedPtr<FObjectMixerEditorList> PinnedListModel = GetListModelPtr().Pin();
+	check(PinnedListModel);
+	
+	return PinnedListModel->GetTreeViewMode();
+}
+
+void SObjectMixerEditorList::SetTreeViewMode(EObjectMixerTreeViewMode InViewMode)
+{
+	if (TSharedPtr<FObjectMixerEditorList> PinnedListModel = GetListModelPtr().Pin())
+	{
+		PinnedListModel->SetTreeViewMode(InViewMode);
+	}
+}
+
 void SObjectMixerEditorList::ToggleFilterActive(const FString& FilterName)
 {
 	if (const TSharedRef<IObjectMixerEditorListFilter>* Match =
@@ -364,7 +448,7 @@ void SObjectMixerEditorList::EvaluateIfRowsPassFilters(const bool bShouldRefresh
 	
 	for (const FObjectMixerEditorListRowPtr& Row : TreeViewRootObjects)
 	{
-		if (Row.IsValid() && Row->GetRowType() == FObjectMixerEditorListRow::SingleItem)
+		if (Row.IsValid() && Row->GetRowType() == FObjectMixerEditorListRow::MatchingObject)
 		{
 			auto Projection = [&Row](const TSharedRef<IObjectMixerEditorListFilter>& Filter)
 			{
@@ -499,7 +583,7 @@ TSharedRef<SWidget> SObjectMixerEditorList::GenerateHeaderRowContextMenu() const
 bool SObjectMixerEditorList::AddUniquePropertyColumnsToHeaderRow(
 	FProperty* Property,
 	const bool bForceIncludeProperty,
-	const TArray<FName>& PropertySkipList)
+	const TSet<FName>& PropertySkipList)
 {
 	if (!ensureAlwaysMsgf(Property, TEXT("%hs: Invalid property passed in. Please ensure only valid properties are passed to this function."), __FUNCTION__))
 	{
@@ -596,33 +680,20 @@ TSharedPtr<SHeaderRow> SObjectMixerEditorList::GenerateHeaderRow()
 {
 	check(ListModelPtr.IsValid());
 	check(HeaderRow);
-
+	
 	TMap<FName, bool> LastVisibleColumns;
+	
 	for (const SHeaderRow::FColumn& Column : HeaderRow->GetColumns())
 	{
 		LastVisibleColumns.Add(Column.ColumnId, Column.bIsVisible);
 	}
-	
+
 	HeaderRow->ClearColumns();
 	ListViewColumns.Empty(ListViewColumns.Num());
-
-	// Property Columns
-	const TObjectPtr<UObjectMixerObjectFilter> SelectedFilter = ListModelPtr.Pin()->GetObjectFilter();
-	if (!SelectedFilter)
-	{
-		UE_LOG(LogObjectMixerEditor, Display, TEXT("%hs: No classes defined in UObjectMixerObjectFilter class."), __FUNCTION__);
-		return nullptr;
-	}
-
-	const EObjectMixerPropertyInheritanceInclusionOptions Options =
-		SelectedFilter->GetObjectMixerPropertyInheritanceInclusionOptions();
-	TArray<UClass*> SpecifiedClasses =
-		SelectedFilter->GetParentAndChildClassesFromSpecifiedClasses(SelectedFilter->GetObjectClassesToFilter(), Options);
 	
-	const TArray<FName> PropertySkipList = SelectedFilter->GetColumnsFilter();
-	const TArray<FName> ForceAddedPropertyList = SelectedFilter->GetForceAddedColumns();
-
-	const bool bShouldIncludeUnsupportedProperties = SelectedFilter->ShouldIncludeUnsupportedProperties();
+	TSet<UClass*> SpecifiedClasses =
+		ListModelPtr.Pin()->GetObjectFilter()->GetParentAndChildClassesFromSpecifiedClasses(
+			ObjectClassesToFilterCache, PropertyInheritanceInclusionOptionsCache);
 	
 	for (const UClass* Class : SpecifiedClasses)
 	{
@@ -630,12 +701,12 @@ TSharedPtr<SHeaderRow> SObjectMixerEditorList::GenerateHeaderRow()
 		{
 			if (FProperty* Property = *FieldIterator)
 			{
-				AddUniquePropertyColumnsToHeaderRow(Property, bShouldIncludeUnsupportedProperties, PropertySkipList);
+				AddUniquePropertyColumnsToHeaderRow(Property, bShouldIncludeUnsupportedPropertiesCache, ColumnsToExcludeCache);
 			}
 		}
 
 		// Check Force Added Columns
-		for (const FName& PropertyName : ForceAddedPropertyList)
+		for (const FName& PropertyName : ForceAddedColumnsCache)
 		{
 			if (FProperty* Property = FindFProperty<FProperty>(Class, PropertyName))
 			{
@@ -660,91 +731,91 @@ TSharedPtr<SHeaderRow> SObjectMixerEditorList::GenerateHeaderRow()
 	AddBuiltinColumnsToHeaderRow();
 
 	// Actually add columns to Header
-	const FText ClickToSortTooltip = LOCTEXT("ClickToSort","Click to sort");
-	
-	const TArray<FName> ColumnsToShowByDefault = SelectedFilter->GetColumnsToShowByDefault();
+	{	
+		const FText ClickToSortTooltip = LOCTEXT("ClickToSort","Click to sort");
 
-	const TSharedRef<SWidget> HeaderMenuContent = GenerateHeaderRowContextMenu();
+		const TSharedRef<SWidget> HeaderMenuContent = GenerateHeaderRowContextMenu();
 
-	for (const FListViewColumnInfo& ColumnInfo : ListViewColumns)
-	{
-		const FText Tooltip = ColumnInfo.PropertyRef ? ColumnInfo.PropertyRef->GetToolTipText() :
-			ColumnInfo.bCanBeSorted ? ClickToSortTooltip : ColumnInfo.PropertyDisplayText;
-		
-		SHeaderRow::FColumn::FArguments Column =
-			SHeaderRow::Column(ColumnInfo.PropertyName)
-			.DefaultLabel(ColumnInfo.PropertyDisplayText)
-			.ToolTipText(Tooltip)
-			.HAlignHeader(EHorizontalAlignment::HAlign_Left)
-		;
-
-		if (ColumnInfo.bUseFixedWidth)
+		for (const FListViewColumnInfo& ColumnInfo : ListViewColumns)
 		{
-			Column.FixedWidth(ColumnInfo.FixedWidth);
-		}
-		else
-		{
-			Column.FillWidth(ColumnInfo.FillWidth);
-		}
+			const FText Tooltip = ColumnInfo.PropertyRef ? ColumnInfo.PropertyRef->GetToolTipText() :
+				ColumnInfo.bCanBeSorted ? ClickToSortTooltip : ColumnInfo.PropertyDisplayText;
+		
+			SHeaderRow::FColumn::FArguments Column =
+				SHeaderRow::Column(ColumnInfo.PropertyName)
+				.DefaultLabel(ColumnInfo.PropertyDisplayText)
+				.ToolTipText(Tooltip)
+				.HAlignHeader(EHorizontalAlignment::HAlign_Left)
+			;
+
+			if (ColumnInfo.bUseFixedWidth)
+			{
+				Column.FixedWidth(ColumnInfo.FixedWidth);
+			}
+			else
+			{
+				Column.FillWidth(ColumnInfo.FillWidth);
+			}
 
-		if (ColumnInfo.bCanBeSorted)
-		{
-			Column.SortMode_Raw(this, &SObjectMixerEditorList::GetSortModeForColumn, ColumnInfo.PropertyName);
-			Column.OnSort_Raw(this, &SObjectMixerEditorList::OnSortColumnCalled);
-		}
+			if (ColumnInfo.bCanBeSorted)
+			{
+				Column.SortMode_Raw(this, &SObjectMixerEditorList::GetSortModeForColumn, ColumnInfo.PropertyName);
+				Column.OnSort_Raw(this, &SObjectMixerEditorList::OnSortColumnCalled);
+			}
 
-		if (ColumnInfo.PropertyType == EListViewColumnType::BuiltIn)
-		{
-			Column.ShouldGenerateWidget(true);
-		}
+			if (ColumnInfo.PropertyType == EListViewColumnType::BuiltIn)
+			{
+				Column.ShouldGenerateWidget(true);
+			}
 
-		if (ColumnInfo.PropertyName.IsEqual(EditorVisibilityColumnName))
-		{
-			Column.HeaderContent()
-			[
-				SNew(SBox)
-				.HAlign(HAlign_Left)
-				.VAlign(VAlign_Center)
-				.Padding(0.f)
+			if (ColumnInfo.PropertyName.IsEqual(EditorVisibilityColumnName))
+			{
+				Column.HeaderContent()
 				[
-					SNew(SImage)
-					.ColorAndOpacity(FSlateColor::UseForeground())
-					.Image(FAppStyle::Get().GetBrush("Level.VisibleIcon16x"))
-				]
-			];
-		}
-		else if (ColumnInfo.PropertyName.IsEqual(EditorVisibilitySoloColumnName))
-		{
-			Column.HeaderContent()
-			[
-				SNew(SBox)
-				.HAlign(HAlign_Left)
-				.VAlign(VAlign_Center)
-				.Padding(0.f)
+					SNew(SBox)
+					.HAlign(HAlign_Left)
+					.VAlign(VAlign_Center)
+					.Padding(0.f)
+					[
+						SNew(SImage)
+						.ColorAndOpacity(FSlateColor::UseForeground())
+						.Image(FAppStyle::Get().GetBrush("Level.VisibleIcon16x"))
+					]
+				];
+			}
+			else if (ColumnInfo.PropertyName.IsEqual(EditorVisibilitySoloColumnName))
+			{
+				Column.HeaderContent()
 				[
-					SNew(SImage)
-					.ColorAndOpacity(FSlateColor::UseForeground())
-					.Image(FAppStyle::Get().GetBrush("MediaAsset.AssetActions.Solo.Small"))
-				]
-			];
-		}
-		else // Add Column Selection Menu widget to all other columns
-		{
-			Column.MenuContent()
-			[
-				HeaderMenuContent
-			];
-		}
+					SNew(SBox)
+					.HAlign(HAlign_Left)
+					.VAlign(VAlign_Center)
+					.Padding(0.f)
+					[
+						SNew(SImage)
+						.ColorAndOpacity(FSlateColor::UseForeground())
+						.Image(FAppStyle::Get().GetBrush("MediaAsset.AssetActions.Solo.Small"))
+					]
+				];
+			}
+			else // Add Column Selection Menu widget to all other columns
+			{
+				Column.MenuContent()
+				[
+					HeaderMenuContent
+				];
+			}
 		
-		HeaderRow->AddColumn(Column);
-		bool bShouldShowColumn = ColumnsToShowByDefault.Contains(ColumnInfo.PropertyName);
+			HeaderRow->AddColumn(Column);
+			bool bShouldShowColumn = ColumnsToShowByDefaultCache.Contains(ColumnInfo.PropertyName);
 
-		if (const bool* Match = LastVisibleColumns.Find(ColumnInfo.PropertyName))
-		{
-			bShouldShowColumn = *Match;
-		}
+			if (const bool* Match = LastVisibleColumns.Find(ColumnInfo.PropertyName))
+			{
+				bShouldShowColumn = *Match;
+			}
 		
-		HeaderRow->SetShowGeneratedColumn(ColumnInfo.PropertyName, bShouldShowColumn);
+			HeaderRow->SetShowGeneratedColumn(ColumnInfo.PropertyName, bShouldShowColumn);
+		}
 	}
 
 	return HeaderRow;
@@ -755,54 +826,207 @@ void SObjectMixerEditorList::SetupFilters()
 	
 }
 
+TSharedRef<SWidget> SObjectMixerEditorList::OnGenerateFilterClassMenu()
+{
+	FMenuBuilder MenuBuilder(true, NULL);
+
+	TArray<UClass*> DerivedClasses;
+	GetDerivedClasses(UObjectMixerObjectFilter::StaticClass(), DerivedClasses, true);
+
+	DerivedClasses.Remove(UObjectMixerObjectFilter::StaticClass());
+	DerivedClasses.Remove(UObjectMixerBlueprintObjectFilter::StaticClass());
+
+	DerivedClasses.Sort([](UClass& A, UClass& B)
+	{
+		return A.GetFName().LexicalLess(B.GetFName());
+	});
+
+	if (DerivedClasses.Num())
+	{
+		const TSharedPtr<FObjectMixerEditorList, ESPMode::ThreadSafe> PinnedList = ListModelPtr.Pin();
+		check(PinnedList);
+		
+		MenuBuilder.BeginSection(NAME_None, LOCTEXT("SelectClassMenuSection", "Select Class"));
+		{
+			for (UClass* DerivedClass : DerivedClasses)
+			{
+				if (IsValid(DerivedClass))
+				{
+					if (DerivedClass->GetName().StartsWith(TEXT("SKEL_")) || DerivedClass->GetName().StartsWith(TEXT("REINST_")))
+					{
+						continue;
+					}
+
+					if (DerivedClass->HasAnyClassFlags(CLASS_Abstract | CLASS_HideDropDown | CLASS_Deprecated))
+					{
+						continue;
+					}
+					
+					MenuBuilder.AddMenuEntry(
+					   FText::FromName(DerivedClass->GetFName()),
+					   FText::GetEmpty(),
+					   FSlateIcon(),
+					   FUIAction(
+						   FExecuteAction::CreateSP(PinnedList.ToSharedRef(), &FObjectMixerEditorList::SetObjectFilterClass, DerivedClass),
+						   FCanExecuteAction::CreateLambda([](){ return true; }),
+						   FIsActionChecked::CreateSP(PinnedList.ToSharedRef(), &FObjectMixerEditorList::IsClassSelected, DerivedClass)
+					   ),
+					   NAME_None,
+					   EUserInterfaceActionType::RadioButton
+				   );
+				}
+			}
+		}
+		MenuBuilder.EndSection();
+	}
+	else
+	{
+		MenuBuilder.AddMenuEntry(LOCTEXT("NoFilterClassesAvailable", "No filter classes available."), FText::GetEmpty(), FSlateIcon(), FUIAction());
+	}
+
+	TSharedRef<SWidget> Widget = MenuBuilder.MakeWidget();
+	FChildren* ChildWidgets = Widget->GetChildren();
+	for (int32 ChildItr = 0; ChildItr < ChildWidgets->Num(); ChildItr++)
+	{
+		const TSharedRef<SWidget>& Child = ChildWidgets->GetChildAt(ChildItr);
+
+		Child->EnableToolTipForceField(false);
+	}
+	Widget->EnableToolTipForceField(false);
+	
+	return Widget;
+}
+
 TSharedRef<SWidget> SObjectMixerEditorList::BuildShowOptionsMenu()
 {
 	FMenuBuilder ShowOptionsMenuBuilder = FMenuBuilder(true, nullptr);
 
-	ShowOptionsMenuBuilder.AddMenuEntry(
-		FText::FromString("Open Generic Object Mixer Instance"),
-		FText::FromString("Open Generic Object Mixer Instance"),
-		FSlateIcon(),
-		FUIAction(FExecuteAction::CreateLambda([]()
-		{
-			FGlobalTabmanager::Get()->TryInvokeTab(FObjectMixerEditorModule::Get().GetTabSpawnerId());
-		})));
+	ShowOptionsMenuBuilder.BeginSection("ListViewOptions", LOCTEXT("FilterClassManagementSection", "Filter Class Management"));
+	{
+		// Filter Class Management Button
+		const TSharedRef<SWidget> FilterClassManagementButton =
+			SNew(SBox)
+			.Padding(8, 0)
+			[
+				SNew(SComboButton)
+				.ToolTipText(LOCTEXT("FilterClassManagementButton_Tooltip", "Select a filter class"))
+				.ContentPadding(FMargin(4, 0.5f))
+				.ComboButtonStyle(&FAppStyle::Get().GetWidgetStyle<FComboButtonStyle>("ComboButton"))
+				.OnGetMenuContent(this, &SObjectMixerEditorList::OnGenerateFilterClassMenu)
+				.ForegroundColor(FStyleColors::Foreground)
+				.MenuPlacement(EMenuPlacement::MenuPlacement_MenuRight)
+				.ButtonContent()
+				[
+					SNew(SHorizontalBox)
+
+					+ SHorizontalBox::Slot()
+					.Padding(0, 1, 4, 0)
+					.AutoWidth()
+					[
+						SNew(SImage)
+						.Image(FAppStyle::Get().GetBrush("Icons.Filter"))
+						.ColorAndOpacity(FSlateColor::UseForeground())
+					]
+
+					+ SHorizontalBox::Slot()
+					.Padding(0, 1, 0, 0)
+					.AutoWidth()
+					[
+						SNew(STextBlock)
+						.Text(LOCTEXT("FilterClassToolbarButton", "Object Filter Class"))
+					]
+				]
+			];
 
-	ShowOptionsMenuBuilder.AddMenuEntry(FText::FromString("Refresh List"), FText::FromString("Refresh"), FSlateIcon(),
-		FUIAction(FExecuteAction::CreateRaw(this, &SObjectMixerEditorList::GenerateTreeView)));
+		ShowOptionsMenuBuilder.AddWidget(FilterClassManagementButton, FText::GetEmpty());
+	}
+	ShowOptionsMenuBuilder.EndSection();
 
-	ShowOptionsMenuBuilder.BeginSection("", LOCTEXT("ShowOptions_ShowSectionHeading", "Show"));
+	// Add List View Mode Options
+	ShowOptionsMenuBuilder.BeginSection("ListViewOptions", LOCTEXT("ListViewOptionsSection", "List View Options"));
 	{
-		// Add show filters
-		auto AddFiltersLambda = [this, &ShowOptionsMenuBuilder](const TSharedRef<IObjectMixerEditorListFilter>& InFilter)
+		// Foreach on uenum
+		const FString EnumPath = "/Script/ObjectMixerEditor.EObjectMixerTreeViewMode";
+		if (const UEnum* EnumPtr = FindObject<UEnum>(nullptr, *EnumPath, true))
 		{
-			const FString& FilterName = InFilter->GetFilterName();
-			
-			ShowOptionsMenuBuilder.AddMenuEntry(
-			   InFilter->GetFilterButtonLabel(),
-			   InFilter->GetFilterButtonToolTip(),
-			   FSlateIcon(),
-			   FUIAction(
-				   FExecuteAction::CreateLambda(
-				   	[this, FilterName]()
-					   {
-						   ToggleFilterActive(FilterName);
-					   }
+			for (int32 EnumItr = 0; EnumItr < EnumPtr->GetMaxEnumValue(); EnumItr++)
+			{
+				EObjectMixerTreeViewMode EnumValue = (EObjectMixerTreeViewMode)EnumItr;
+				
+				ShowOptionsMenuBuilder.AddMenuEntry(
+					EnumPtr->GetDisplayNameTextByIndex(EnumItr),
+					EnumPtr->GetToolTipTextByIndex(EnumItr),
+					FSlateIcon(),
+					FUIAction(
+						FExecuteAction::CreateLambda([this, EnumValue]()
+						{
+							SetTreeViewMode(EnumValue);
+						}),
+						FCanExecuteAction::CreateLambda([](){ return true; }),
+						FIsActionChecked::CreateLambda([this, EnumValue]()
+						{
+							return GetTreeViewMode() == EnumValue;
+						})
 					),
-				   FCanExecuteAction(),
-				   FIsActionChecked::CreateSP( InFilter, &IObjectMixerEditorListFilter::GetIsFilterActive )
-			   ),
-			   NAME_None,
-			   EUserInterfaceActionType::ToggleButton
-		   );
-		};
+					NAME_None,
+					EUserInterfaceActionType::RadioButton
+				);
+			}
+		}
+	}
+	ShowOptionsMenuBuilder.EndSection();
+
+	ShowOptionsMenuBuilder.BeginSection("MiscOptionsSection", LOCTEXT("MiscOptionsSection","Misc"));
+	{
+		ShowOptionsMenuBuilder.AddMenuEntry(
+		   FText::FromString("Open Generic Object Mixer Instance"),
+		   FText::FromString("Open Generic Object Mixer Instance"),
+		   FSlateIcon(),
+		   FUIAction(FExecuteAction::CreateLambda([]()
+		   {
+			   FGlobalTabmanager::Get()->TryInvokeTab(FObjectMixerEditorModule::Get().GetTabSpawnerId());
+		   })));
+
+		ShowOptionsMenuBuilder.AddMenuEntry(FText::FromString("Refresh List"), FText::FromString("Refresh"), FSlateIcon(),
+			FUIAction(FExecuteAction::CreateRaw(this, &SObjectMixerEditorList::GenerateTreeView)));
+	}
+	ShowOptionsMenuBuilder.EndSection();
 
-		for (const TSharedRef<IObjectMixerEditorListFilter>& Filter : ShowFilters)
+	if (ShowFilters.Num())
+	{
+		ShowOptionsMenuBuilder.BeginSection("", LOCTEXT("ShowOptions_ShowSectionHeading", "Show"));
 		{
-			AddFiltersLambda(Filter);
+			// Add show filters
+			auto AddFiltersLambda = [this, &ShowOptionsMenuBuilder](const TSharedRef<IObjectMixerEditorListFilter>& InFilter)
+			{
+				const FString& FilterName = InFilter->GetFilterName();
+			
+				ShowOptionsMenuBuilder.AddMenuEntry(
+				   InFilter->GetFilterButtonLabel(),
+				   InFilter->GetFilterButtonToolTip(),
+				   FSlateIcon(),
+				   FUIAction(
+					   FExecuteAction::CreateLambda(
+						   [this, FilterName]()
+						   {
+							   ToggleFilterActive(FilterName);
+						   }
+						),
+					   FCanExecuteAction(),
+					   FIsActionChecked::CreateSP( InFilter, &IObjectMixerEditorListFilter::GetIsFilterActive )
+				   ),
+				   NAME_None,
+				   EUserInterfaceActionType::ToggleButton
+			   );
+			};
+
+			for (const TSharedRef<IObjectMixerEditorListFilter>& Filter : ShowFilters)
+			{
+				AddFiltersLambda(Filter);
+			}
 		}
+		ShowOptionsMenuBuilder.EndSection();
 	}
-	ShowOptionsMenuBuilder.EndSection();
 
 	return ShowOptionsMenuBuilder.MakeWidget();
 }
@@ -843,6 +1067,158 @@ void SObjectMixerEditorList::OnListViewSearchTextChanged(const FText& Text)
 	ExecuteListViewSearchOnAllRows(Text.ToString(), true);
 }
 
+void SObjectMixerEditorList::CacheTreeState()
+{
+	struct Local
+	{
+		static void RecursivelyCacheTreeState(
+			const TArray<FObjectMixerEditorListRowPtr>& InObjects,
+			TMap<FString, bool>& TreeItemExpansionStateCache,
+			TSharedPtr<STreeView<FObjectMixerEditorListRowPtr>> TreeViewPtr)
+		{
+			for (const TSharedPtr<FObjectMixerEditorListRow>& TreeViewItem : InObjects)
+			{
+				const FString ObjectName = TreeViewItem->GetDisplayName().ToString();
+
+				if (!ObjectName.IsEmpty())
+				{
+					TreeItemExpansionStateCache.Add(ObjectName, TreeViewPtr->IsItemExpanded(TreeViewItem));
+				}
+
+				RecursivelyCacheTreeState(TreeViewItem->GetChildRows(), TreeItemExpansionStateCache, TreeViewPtr);
+			}
+		}
+	};
+
+	Local::RecursivelyCacheTreeState(TreeViewRootObjects, TreeItemExpansionStateCache, TreeViewPtr);
+}
+
+void SObjectMixerEditorList::RestoreTreeState(const bool bFlushCache)
+{
+	struct Local
+	{
+		static void RecursivelyRestoreTreeState(
+			const TArray<FObjectMixerEditorListRowPtr>& InObjects,
+			TMap<FString, bool>& TreeItemExpansionStateCache,
+			TSharedPtr<STreeView<FObjectMixerEditorListRowPtr>> TreeViewPtr,
+			const bool bExpandByDefault)
+		{
+			for (const TSharedPtr<FObjectMixerEditorListRow>& TreeViewItem : InObjects)
+			{
+				const FString ObjectName = TreeViewItem->GetDisplayName().ToString();
+
+				if (!ObjectName.IsEmpty())
+				{
+					if (const bool* bExpansionStatePtr = TreeItemExpansionStateCache.Find(ObjectName))
+					{
+						TreeViewPtr->SetItemExpansion(TreeViewItem, *bExpansionStatePtr);
+					}
+					else
+					{
+						TreeViewPtr->SetItemExpansion(TreeViewItem, bExpandByDefault);
+					}
+				}
+
+				RecursivelyRestoreTreeState(
+					TreeViewItem->GetChildRows(), TreeItemExpansionStateCache, TreeViewPtr, bExpandByDefault);
+			}
+		}
+	};
+	
+	bool bExpandByDefault = true;
+	if (const UObjectMixerEditorProjectSettings* const Settings = GetDefault<UObjectMixerEditorProjectSettings>())
+	{
+		bExpandByDefault = Settings->bExpandTreeViewItemsByDefault;
+	}
+	
+	Local::RecursivelyRestoreTreeState(
+		TreeViewRootObjects, TreeItemExpansionStateCache, TreeViewPtr, bExpandByDefault);
+
+	if (bFlushCache)
+	{
+		TreeItemExpansionStateCache.Empty();
+	}
+}
+
+void SObjectMixerEditorList::BuildPerformanceCacheAndGenerateHeaderIfNeeded()
+{
+	// If any of the following overrides change, we need to regenerate the header row. Otherwise skip regeneration for performance reasons.
+	// GetObjectClassesToFilter, GetColumnsToShowByDefault, GetColumnsToExclude,
+	// GetForceAddedColumns, GetObjectMixerPropertyInheritanceInclusionOptions, ShouldIncludeUnsupportedProperties
+	bool bNeedToGenerateHeaders = false;
+	
+	const TObjectPtr<UObjectMixerObjectFilter> SelectedFilter = ListModelPtr.Pin()->GetObjectFilter();
+	if (!SelectedFilter)
+	{
+		UE_LOG(LogObjectMixerEditor, Display, TEXT("%hs: No classes defined in UObjectMixerObjectFilter class."), __FUNCTION__);
+		return;
+	}
+
+	if (const TSet<UClass*> ObjectClassesToFilter = SelectedFilter->GetObjectClassesToFilter();
+		ObjectClassesToFilter.Num() != ObjectClassesToFilterCache.Num() ||
+		ObjectClassesToFilter.Difference(ObjectClassesToFilterCache).Num() > 0 || ObjectClassesToFilterCache.Difference(ObjectClassesToFilter).Num() > 0)
+	{
+		ObjectClassesToFilterCache = ObjectClassesToFilter;
+		if (!bNeedToGenerateHeaders)
+		{
+			bNeedToGenerateHeaders = true;
+		}
+	}
+	if (const TSet<FName> ColumnsToShowByDefault = SelectedFilter->GetColumnsToShowByDefault();
+		ColumnsToShowByDefault.Num() != ColumnsToShowByDefaultCache.Num() ||
+		ColumnsToShowByDefault.Difference(ColumnsToShowByDefaultCache).Num() > 0 || ColumnsToShowByDefaultCache.Difference(ColumnsToShowByDefault).Num() > 0)
+	{
+		ColumnsToShowByDefaultCache = ColumnsToShowByDefault;
+		if (!bNeedToGenerateHeaders)
+		{
+			bNeedToGenerateHeaders = true;
+		}
+	}
+	if (const TSet<FName> ColumnsToExclude = SelectedFilter->GetColumnsToExclude();
+		ColumnsToExclude.Num() != ColumnsToExcludeCache.Num() ||
+		ColumnsToExclude.Difference(ColumnsToExcludeCache).Num() > 0 || ColumnsToExcludeCache.Difference(ColumnsToExclude).Num() > 0)
+	{
+		ColumnsToExcludeCache = ColumnsToExclude;
+		if (!bNeedToGenerateHeaders)
+		{
+			bNeedToGenerateHeaders = true;
+		}
+	}
+	if (const TSet<FName> ForceAddedColumns = SelectedFilter->GetForceAddedColumns();
+		ForceAddedColumns.Num() != ForceAddedColumnsCache.Num() ||
+		ForceAddedColumns.Difference(ForceAddedColumnsCache).Num() > 0 || ForceAddedColumnsCache.Difference(ForceAddedColumns).Num() > 0)
+	{
+		ForceAddedColumnsCache = ForceAddedColumns;
+		if (!bNeedToGenerateHeaders)
+		{
+			bNeedToGenerateHeaders = true;
+		}
+	}
+	if (const EObjectMixerInheritanceInclusionOptions PropertyInheritanceInclusionOptions = SelectedFilter->GetObjectMixerPropertyInheritanceInclusionOptions(); 
+		PropertyInheritanceInclusionOptions != PropertyInheritanceInclusionOptionsCache)
+	{
+		PropertyInheritanceInclusionOptionsCache = PropertyInheritanceInclusionOptions;
+		if (!bNeedToGenerateHeaders)
+		{
+			bNeedToGenerateHeaders = true;
+		}
+	}
+	if (const bool bShouldIncludeUnsupportedProperties = SelectedFilter->ShouldIncludeUnsupportedProperties(); 
+		bShouldIncludeUnsupportedProperties != bShouldIncludeUnsupportedPropertiesCache)
+	{
+		bShouldIncludeUnsupportedPropertiesCache = bShouldIncludeUnsupportedProperties;
+		if (!bNeedToGenerateHeaders)
+		{
+			bNeedToGenerateHeaders = true;
+		}
+	}
+	
+	if (bNeedToGenerateHeaders)
+	{
+		GenerateHeaderRow();
+	}
+}
+
 void SObjectMixerEditorList::GenerateTreeView()
 {
 	check(ListModelPtr.IsValid());
@@ -851,35 +1227,27 @@ void SObjectMixerEditorList::GenerateTreeView()
 	{
 		return;
 	}
+
+	CacheTreeState();
 	
 	FlushMemory(true);
-	
-	GenerateHeaderRow();
 
-	TArray<UClass*> AcceptableClasses = ListModelPtr.Pin()->GetObjectClasses();
+	BuildPerformanceCacheAndGenerateHeaderIfNeeded();
 
 	check(GEditor);
 	const UWorld* EditorWorld = GEditor->GetEditorWorldContext().World();
 
-	// Bind delegates to keep view up to date
-	OnActorSpawnedHandle = EditorWorld->AddOnActorSpawnedHandler(
-			FOnActorSpawned::FDelegate::CreateRaw(
-			this, &SObjectMixerEditorList::OnActorSpawnedOrDestroyed));
-	OnActorDestroyedHandle = EditorWorld->AddOnActorDestroyedHandler(
-		FOnActorDestroyed::FDelegate::CreateRaw(
-		this, &SObjectMixerEditorList::OnActorSpawnedOrDestroyed));
-
 	// Find valid matching objects
+	TArray<UObject*> MatchingObjects;
 	for (TObjectIterator<UObject> ObjectIterator; ObjectIterator; ++ObjectIterator)
 	{
 		if (UObject* Object = *ObjectIterator; IsValid(Object))
 		{
-			if (Object->GetWorld() == EditorWorld)
+			if (UWorld* ObjectWorld = Object->GetWorld(); ObjectWorld == EditorWorld)
 			{
-				//Object->GetWorld()->GetLevel(0).
 				bool bIsAcceptableClass = false;
 
-				for (UClass* Class : AcceptableClasses)
+				for (UClass* Class : ObjectClassesToFilterCache)
 				{
 					if (Object->IsA(Class))
 					{
@@ -890,16 +1258,141 @@ void SObjectMixerEditorList::GenerateTreeView()
 
 				if (bIsAcceptableClass)
 				{
-					TreeViewRootObjects.Add(
-						MakeShared<FObjectMixerEditorListRow>(
-							Object, FObjectMixerEditorListRow::SingleItem, SharedThis(this), nullptr)
-					);
+					MatchingObjects.Add(Object);
 				}
 			}
 		}
 	}
 
+	// A quick lookup for objects that already exist in the list.
+	// Helpful to avoid double-generating rows when considering parent->child hierarchy.
+	TMap<UObject*, TSharedRef<FObjectMixerEditorListRow>> CreatedObjectMap;
+	TMap<FName, TSharedRef<FObjectMixerEditorListRow>> FolderMap;
+	for (UObject* Object : MatchingObjects)
+	{
+		if (!CreatedObjectMap.Contains(Object)) // Ensure we don't double-generate container objects
+		{
+			TSharedRef<FObjectMixerEditorListRow> TopLevelRow = MakeShared<FObjectMixerEditorListRow>(
+				Object, FObjectMixerEditorListRow::MatchingObject, SharedThis(this));
+
+			CreatedObjectMap.Add(Object, TopLevelRow);
+
+			// If the view is not in flat mode, we need to consider the hierarchy of outliner folders/attach parents as desired
+			if (ListModelPtr.Pin()->GetTreeViewMode() != EObjectMixerTreeViewMode::Flat)
+			{
+				AActor* BaseActor = Cast<AActor>(Object);
+				
+				if (!BaseActor)
+				{
+					if (Object->IsA(UActorComponent::StaticClass()))
+					{
+						BaseActor = Object->GetTypedOuter<AActor>();
+
+						// If it's not flat or folder view mode, we need to find or create the container object for the actor that owns the matching component
+						if (ListModelPtr.Pin()->GetTreeViewMode() != EObjectMixerTreeViewMode::Folder)
+						{
+							TSharedPtr<FObjectMixerEditorListRow> OwningActorRow;
+
+							if (const TSharedRef<FObjectMixerEditorListRow>* Match = CreatedObjectMap.Find(BaseActor))
+							{
+								OwningActorRow = *Match;
+							}
+							else
+							{
+								OwningActorRow = MakeShared<FObjectMixerEditorListRow>(
+									BaseActor, FObjectMixerEditorListRow::ContainerObject, SharedThis(this));
+
+								CreatedObjectMap.Add(BaseActor, OwningActorRow.ToSharedRef());
+							}
+
+							if (OwningActorRow)
+							{
+								OwningActorRow->AddToChildRows(TopLevelRow);
+													
+								TopLevelRow = OwningActorRow.ToSharedRef();
+							}
+						}
+					}
+				}
+
+				if (BaseActor)
+				{
+					while (AActor* AttachParent = BaseActor->GetAttachParentActor())
+					{
+						// Make a row for each attach parent up the chain until we reach the top if not in flat/folder mode
+						if (ListModelPtr.Pin()->GetTreeViewMode() != EObjectMixerTreeViewMode::Folder)
+						{
+							TSharedPtr<FObjectMixerEditorListRow> OwningActorRow;
+
+							if (const TSharedRef<FObjectMixerEditorListRow>* Match = CreatedObjectMap.Find(AttachParent))
+							{
+								OwningActorRow = *Match;
+							}
+							else
+							{
+								OwningActorRow = MakeShared<FObjectMixerEditorListRow>(
+									AttachParent, FObjectMixerEditorListRow::ContainerObject, SharedThis(this));
+
+								CreatedObjectMap.Add(AttachParent, OwningActorRow.ToSharedRef());
+							}
+
+							if (OwningActorRow)
+							{
+								OwningActorRow->AddToChildRows(TopLevelRow);
+												
+								TopLevelRow = OwningActorRow.ToSharedRef();
+							}
+						}
+
+						BaseActor = AttachParent;
+					}
+
+					// Now consider folder hierarchy for the base actor if desired
+					if (ListModelPtr.Pin()->GetTreeViewMode() == EObjectMixerTreeViewMode::Folder ||
+						ListModelPtr.Pin()->GetTreeViewMode() == EObjectMixerTreeViewMode::FolderObjectSubObject)
+					{
+						FFolder BaseActorFolder = BaseActor->GetFolder();
+
+						while (!BaseActorFolder.IsNone())
+						{
+							TSharedPtr<FObjectMixerEditorListRow> FolderRow;
+							
+							if (const TSharedRef<FObjectMixerEditorListRow>* Match = FolderMap.Find(BaseActorFolder.GetPath()))
+							{
+								FolderRow = *Match;
+							}
+							else
+							{
+								FolderRow =
+									MakeShared<FObjectMixerEditorListRow>(
+										nullptr, FObjectMixerEditorListRow::Folder, SharedThis(this),
+										FText::FromName(BaseActorFolder.GetLeafName()));
+
+								FolderMap.Add(BaseActorFolder.GetPath(), FolderRow.ToSharedRef());
+							}
+
+							if (FolderRow)
+							{
+								FolderRow->AddToChildRows(TopLevelRow);
+													
+								TopLevelRow = FolderRow.ToSharedRef();
+							}
+
+							BaseActorFolder = BaseActorFolder.GetParent();
+						}
+					}
+				}
+			}
+
+			TreeViewRootObjects.AddUnique(TopLevelRow);
+		}
+	}
+
+	TreeViewRootObjects.StableSort(SortByTypeThenName);
+
 	RefreshList();
+
+	RestoreTreeState();
 }
 
 void SObjectMixerEditorList::FindVisibleTreeViewObjects()
@@ -908,7 +1401,7 @@ void SObjectMixerEditorList::FindVisibleTreeViewObjects()
 
 	for (const TSharedPtr<FObjectMixerEditorListRow>& Row : TreeViewRootObjects)
 	{
-		if (Row->ShouldBeVisible())
+		if (Row->ShouldRowWidgetBeVisible())
 		{
 			VisibleTreeViewObjects.Add(Row);
 		}
@@ -943,7 +1436,7 @@ void SObjectMixerEditorList::OnRowChildExpansionChange(FObjectMixerEditorListRow
 		{
 			if (bIsExpanded)
 			{
-				if (Row->GetRowType() == FObjectMixerEditorListRow::Group)
+				if (Row->GetRowType() == FObjectMixerEditorListRow::Folder)
 				{
 					Row->SetShouldExpandAllChildren(true);
 				}
diff --git a/Engine/Source/Editor/ObjectMixer/ObjectMixer/Private/Views/List/SObjectMixerEditorListRow.cpp b/Engine/Source/Editor/ObjectMixer/ObjectMixer/Private/Views/List/SObjectMixerEditorListRow.cpp
--- a/Engine/Source/Editor/ObjectMixer/ObjectMixer/Private/Views/List/SObjectMixerEditorListRow.cpp
+++ b/Engine/Source/Editor/ObjectMixer/ObjectMixer/Private/Views/List/SObjectMixerEditorListRow.cpp
@@ -2,25 +2,24 @@
 
 #include "Views/List/SObjectMixerEditorListRow.h"
 
+#include "EditorClassUtils.h"
 #include "ObjectMixerEditorStyle.h"
 #include "Views/List/ObjectMixerEditorList.h"
 #include "Views/List/SObjectMixerEditorList.h"
 
 #include "Customizations/ColorStructCustomization.h"
+#include "Engine/Blueprint.h"
+#include "GameFramework/Actor.h"
 #include "Input/DragAndDrop.h"
 #include "ISinglePropertyView.h"
-#include "IStructureDetailsView.h"
 #include "Modules/ModuleManager.h"
-#include "PropertyEditorDelegates.h"
 #include "PropertyEditorModule.h"
 #include "PropertyHandle.h"
 #include "ScopedTransaction.h"
 #include "Styling/AppStyle.h"
 #include "Styling/StyleColors.h"
-#include "Widgets/Input/SCheckBox.h"
 #include "Widgets/Layout/SBox.h"
 #include "Widgets/SBoxPanel.h"
-#include "Widgets/SOverlay.h"
 #include "Widgets/Layout/SScaleBox.h"
 #include "Widgets/Text/STextBlock.h"
 
@@ -80,9 +79,6 @@ void SObjectMixerEditorListRow::Construct(
 	Item = InRow;
 	const FObjectMixerEditorListRowPtr PinnedItem = Item.Pin();
 
-	// Set up flash animation
-	FlashAnimation = FCurveSequence(0.f, FlashAnimationDuration, ECurveEaseFunction::QuadInOut);
-
 	SMultiColumnTableRow<FObjectMixerEditorListRowPtr>::Construct(
 		FSuperRowType::FArguments()
 		.Padding(1.0f)
@@ -102,54 +98,53 @@ void SObjectMixerEditorListRow::Construct(
 	VisibleNotHoveredBrush = FAppStyle::Get().GetBrush(VisibleNotHoveredBrushName);
 	NotVisibleHoveredBrush = FAppStyle::Get().GetBrush(NotVisibleHoveredBrushName);
 	NotVisibleNotHoveredBrush = FAppStyle::Get().GetBrush(NotVisibleNotHoveredBrushName);
-
-	if (PinnedItem->GetShouldFlashOnScrollIntoView())
-	{
-		FlashRow();
-
-		PinnedItem->SetShouldFlashOnScrollIntoView(false);
-	}
 }
 
 TSharedRef<SWidget> SObjectMixerEditorListRow::GenerateWidgetForColumn(const FName& InColumnName)
 {
 	const FObjectMixerEditorListRowPtr PinnedItem = Item.Pin();
 
-	const TSharedPtr<SWidget> CellWidget = GenerateCells(InColumnName, PinnedItem);
-
-	const TSharedRef<SImage> FlashImage = SNew(SImage)
-										.Image(new FSlateColorBrush(FStyleColors::White))
-										.Visibility_Raw(
-											this, &SObjectMixerEditorListRow::GetFlashImageVisibility)
-										.ColorAndOpacity_Raw(
-											this, &SObjectMixerEditorListRow::GetFlashImageColorAndOpacity);
-
-	FlashImages.Add(FlashImage);
-
-	return SNew(SOverlay)
-			.Visibility(EVisibility::SelfHitTestInvisible)
-
-			+ SOverlay::Slot()
-			[
-				FlashImage
-			]
-
-			+ SOverlay::Slot()
-			[
-				SNew(SBorder)
-					.HAlign(HAlign_Fill)
-					.VAlign(VAlign_Center)
-					.BorderImage(GetBorderImage(PinnedItem->GetRowType()))
+	if (const TSharedPtr<SWidget> CellWidget = GenerateCells(InColumnName, PinnedItem))
+	{
+		if (InColumnName == SObjectMixerEditorList::ItemNameColumnName)
+		{
+			// The first column gets the tree expansion arrow for this row
+			return SNew(SBox)
+				.MinDesiredHeight(20)
 				[
-					CellWidget.ToSharedRef()
-				]
-			];
+					SNew( SHorizontalBox )
+
+					+SHorizontalBox::Slot()
+					.AutoWidth()
+					.Padding(6, 0, 0, 0)
+					[
+						SNew( SExpanderArrow, SharedThis(this) ).IndentAmount(12)
+					]
+
+					+SHorizontalBox::Slot()
+					.FillWidth(1.0f)
+					[
+						CellWidget.ToSharedRef()
+					]
+				];
+		}
+		
+		return SNew(SBorder)
+				   .HAlign(HAlign_Fill)
+				   .VAlign(VAlign_Center)
+				   .BorderImage(GetBorderImage(PinnedItem->GetRowType()))
+			   [
+				   CellWidget.ToSharedRef()
+			   ];
+	}
+
+	return SNullWidget::NullWidget;
 }
 
 void SObjectMixerEditorListRow::OnMouseEnter(const FGeometry& MyGeometry, const FPointerEvent& MouseEvent)
 {
 	bIsHovered = true;
-
+	
 	SMultiColumnTableRow<FObjectMixerEditorListRowPtr>::OnMouseEnter(MyGeometry, MouseEvent);
 }
 
@@ -163,13 +158,6 @@ void SObjectMixerEditorListRow::OnMouseLeave(const FPointerEvent& MouseEvent)
 SObjectMixerEditorListRow::~SObjectMixerEditorListRow()
 {
 	Item.Reset();
-	HoverToolTip.Reset();
-
-	FlashImages.Empty();
-
-	ValueChildInputWidget.Reset();
-
-	HoverableWidgetsPtr.Reset();
 }
 
 // FReply SObjectMixerEditorListRow::HandleDragDetected(const FGeometry& MyGeometry, const FPointerEvent& MouseEvent)
@@ -291,25 +279,11 @@ SObjectMixerEditorListRow::~SObjectMixerEditorListRow()
 // 	return FReply::Handled();
 // }
 
-void SObjectMixerEditorListRow::FlashRow()
-{
-	FlashAnimation.Play(this->AsShared());
-}
-
 bool SObjectMixerEditorListRow::IsVisible() const
 {
 	if (const TSharedPtr<FObjectMixerEditorListRow> PinnedItem = Item.Pin())
 	{
-		if (const TSharedPtr<SObjectMixerEditorList> PinnedListView = PinnedItem->GetListViewPtr().Pin())
-		{
-			if (const TSharedPtr<FObjectMixerEditorList> PinnedListModel = PinnedListView->GetListModelPtr().Pin())
-			{
-				if (const UObjectMixerObjectFilter* Filter = PinnedListModel->GetObjectFilter())
-				{
-					return Filter->GetRowEditorVisibility(PinnedItem->GetObject());
-				}
-			}
-		}
+		return PinnedItem->GetObjectVisibility();
 	}
 
 	return false;
@@ -356,70 +330,69 @@ const FSlateBrush* SObjectMixerEditorListRow::GetVisibilityBrush() const
 	return bIsHovered ? NotVisibleHoveredBrush : NotVisibleNotHoveredBrush;
 }
 
-EVisibility SObjectMixerEditorListRow::GetFlashImageVisibility() const
-{
-	return FlashAnimation.IsPlaying() ? EVisibility::SelfHitTestInvisible : EVisibility::Hidden;
-}
-
-FSlateColor SObjectMixerEditorListRow::GetFlashImageColorAndOpacity() const
-{
-	if (FlashAnimation.IsPlaying())
-	{
-		// This equation modulates the alpha into a parabolic curve 
-		const float Progress = FMath::Abs(FMath::Abs((FlashAnimation.GetLerp() - 0.5f) * 2) - 1);
-		return FLinearColor::LerpUsingHSV(FLinearColor::Transparent, FlashColor, Progress);
-	}
-
-	return FLinearColor::Transparent;
-}
-
 const FSlateBrush* SObjectMixerEditorListRow::GetBorderImage(
 	const FObjectMixerEditorListRow::EObjectMixerEditorListRowType InRowType)
 {
-	switch (InRowType)
-	{
-	case FObjectMixerEditorListRow::Group:
-		return FObjectMixerEditorStyle::Get().GetBrush("ObjectMixerEditor.CommandGroupBorder");
-
-	default:
-		return FObjectMixerEditorStyle::Get().GetBrush("ObjectMixerEditor.DefaultBorder");
-	}
+	return FObjectMixerEditorStyle::Get().GetBrush("ObjectMixerEditor.DefaultBorder");
 }
 
-TSharedRef<SWidget> SObjectMixerEditorListRow::GenerateCells(
+TSharedPtr<SWidget> SObjectMixerEditorListRow::GenerateCells(
 	const FName& InColumnName, const TSharedPtr<FObjectMixerEditorListRow> PinnedItem)
 {
 	check(PinnedItem.IsValid());
 	
 	if (InColumnName.IsEqual(SObjectMixerEditorList::ItemNameColumnName))
 	{
-		FText RowName;
+		TSharedRef<SHorizontalBox> HBox = SNew(SHorizontalBox);
+		
+		HBox->AddSlot()
+		.AutoWidth()
+		[
+			SNew(SImage)
+			.Image_Lambda([PinnedItem]()
+			{
+				return PinnedItem->GetObjectIconBrush();
+			})
+			.ColorAndOpacity(FSlateColor::UseForeground())
+		];
 
-		if (const UObjectMixerObjectFilter* Filter = PinnedItem->GetObjectFilter())
+		bool bNeedsStandardTextBlock = true;
+		const FText DisplayName = PinnedItem->GetDisplayName();
+		if (TObjectPtr<UObject> Object = PinnedItem->GetObject())
 		{
-			RowName = Filter->GetRowDisplayName(PinnedItem->GetObject());
+			if (UClass* ActorClass = Object->GetClass())
+			{
+				if (UBlueprint* AsBlueprint = UBlueprint::GetBlueprintFromClass(ActorClass))
+				{
+					bNeedsStandardTextBlock = false;
+					
+					FEditorClassUtils::FSourceLinkParams SourceLinkParams;
+					SourceLinkParams.Object = Object;
+					SourceLinkParams.bUseDefaultFormat = false;
+					SourceLinkParams.bUseFormatIfNoLink = true;
+					SourceLinkParams.BlueprintFormat = &DisplayName;
+
+					HBox->AddSlot()
+					.Padding(FMargin(10.0, 0, 0, 0))
+					[
+						FEditorClassUtils::GetSourceLink(ActorClass, SourceLinkParams)
+					];
+				}
+			}
 		}
 
-		TSharedRef<SHorizontalBox> HBox = SNew(SHorizontalBox);
-		
-		if (const FSlateBrush* Icon = PinnedItem->GetObjectIconBrush())
+		if (bNeedsStandardTextBlock)
 		{
 			HBox->AddSlot()
-			.AutoWidth()
+			.Padding(FMargin(10.0, 0, 0, 0))
 			[
-				SNew(SImage).Image(Icon)
+				SNew(STextBlock)
+				.Visibility(EVisibility::Visible)
+				.Justification(ETextJustify::Left)
+				.Text(DisplayName)
+				.ToolTipText(DisplayName)
 			];
 		}
-
-		HBox->AddSlot()
-		.Padding(FMargin(10.0, 0, 0, 0))
-		[
-			SNew(STextBlock)
-			.Visibility(EVisibility::Visible)
-			.Justification(ETextJustify::Left)
-			.Text(RowName)
-			.ToolTip(HoverToolTip)
-		];
 		
 		return SNew(SBox)
 				.Visibility(EVisibility::SelfHitTestInvisible)
@@ -433,6 +406,16 @@ TSharedRef<SWidget> SObjectMixerEditorListRow::GenerateCells(
 
 	if (InColumnName.IsEqual(SObjectMixerEditorList::EditorVisibilityColumnName))
 	{
+		if (PinnedItem->GetRowType() == FObjectMixerEditorListRow::None)
+		{
+			return nullptr;
+		}
+
+		if (PinnedItem->GetObject() && !PinnedItem->GetObject()->IsA(AActor::StaticClass()))
+		{
+			return nullptr;
+		}
+		
 		return SNew(SBox)
 				.HAlign(HAlign_Left)
 				.VAlign(VAlign_Center)
@@ -458,8 +441,9 @@ TSharedRef<SWidget> SObjectMixerEditorListRow::GenerateCells(
 
 								return FReply::Handled();
 							}
-							
-							PinnedItem->SetObjectVisibility(!bIsVisible);
+
+							// Set Visibility Recursively
+							PinnedItem->SetObjectVisibility(!bIsVisible, true);
 
 							return FReply::Handled();
 						}
@@ -470,6 +454,16 @@ TSharedRef<SWidget> SObjectMixerEditorListRow::GenerateCells(
 
 	if (InColumnName.IsEqual(SObjectMixerEditorList::EditorVisibilitySoloColumnName))
 	{
+		if (PinnedItem->GetRowType() == FObjectMixerEditorListRow::None)
+		{
+			return nullptr;
+		}
+
+		if (PinnedItem->GetObject() && !PinnedItem->GetObject()->IsA(AActor::StaticClass()))
+		{
+			return nullptr;
+		}
+		
 		return SNew(SBox)
 				.HAlign(HAlign_Left)
 				.VAlign(VAlign_Center)
@@ -490,7 +484,7 @@ TSharedRef<SWidget> SObjectMixerEditorListRow::GenerateCells(
 								const bool bIsRowSolo = PinnedItem->IsThisRowSolo();
 								for (const TSharedPtr<FObjectMixerEditorListRow>& TreeItem : PinnedListView->GetTreeViewItems())
 								{
-									TreeItem->SetObjectVisibility(bIsRowSolo);
+									TreeItem->SetObjectVisibility(bIsRowSolo, true);
 								}
 
 								if (bIsRowSolo)
@@ -499,7 +493,7 @@ TSharedRef<SWidget> SObjectMixerEditorListRow::GenerateCells(
 								}
 								else
 								{
-									PinnedItem->SetObjectVisibility(true);
+									PinnedItem->SetObjectVisibility(true, true);
 									PinnedItem->SetThisAsSoloRow();
 								}
 
@@ -513,23 +507,24 @@ TSharedRef<SWidget> SObjectMixerEditorListRow::GenerateCells(
 			;
 	}
 	
-	TSharedPtr<SWidget> Widget;
 	if (UObject* ObjectRef = PinnedItem->GetObject())
 	{
 		FPropertyEditorModule& PropertyEditorModule = FModuleManager::LoadModuleChecked<FPropertyEditorModule>("PropertyEditor");
 		{
 			FSinglePropertyParams Params;
 			Params.NamePlacement = EPropertyNamePlacement::Hidden;
-
-			if (const TSharedPtr<ISinglePropertyView> SinglePropertyView =
+			
+			const TSharedPtr<ISinglePropertyView> SinglePropertyView =
 				PropertyEditorModule.CreateSingleProperty(ObjectRef, InColumnName, Params
-			))
+			);
+
+			if (SinglePropertyView)
 			{
 				if (const TSharedPtr<IPropertyHandle> Handle = SinglePropertyView->GetPropertyHandle())
 				{
 					if (const FProperty* Property = Handle->GetProperty())
 					{
-						// Do this for all selected rows with a similar property
+						// Simultaneously edit all selected rows with a similar property
 						FSimpleDelegate OnPropertyValueChanged =
 							FSimpleDelegate::CreateRaw(
 								this,
@@ -537,25 +532,20 @@ TSharedRef<SWidget> SObjectMixerEditorListRow::GenerateCells(
 					
 						SinglePropertyView->SetOnPropertyValueChanged(OnPropertyValueChanged);
 
-						Widget = SinglePropertyView;
+						return SNew(SBox)
+								.Visibility(EVisibility::SelfHitTestInvisible)
+								.HAlign(HAlign_Fill)
+								.VAlign(VAlign_Center)
+								[
+									SinglePropertyView.ToSharedRef()
+								];
 					}
 				}
 			}
 		}
 	}
-	
-	if (Widget.IsValid())
-	{
-		return SNew(SBox)
-				.Visibility(EVisibility::SelfHitTestInvisible)
-				.HAlign(HAlign_Fill)
-				.VAlign(VAlign_Center)
-				[
-					Widget.ToSharedRef()
-				];
-	}
 
-	return SNullWidget::NullWidget;
+	return nullptr;
 }
 
 void SObjectMixerEditorListRow::OnPropertyChanged(const FProperty* Property, void* ContainerWithChangedProperty)
diff --git a/Engine/Source/Editor/ObjectMixer/ObjectMixer/Private/Views/MainPanel/SObjectMixerEditorMainPanel.cpp b/Engine/Source/Editor/ObjectMixer/ObjectMixer/Private/Views/MainPanel/SObjectMixerEditorMainPanel.cpp
--- a/Engine/Source/Editor/ObjectMixer/ObjectMixer/Private/Views/MainPanel/SObjectMixerEditorMainPanel.cpp
+++ b/Engine/Source/Editor/ObjectMixer/ObjectMixer/Private/Views/MainPanel/SObjectMixerEditorMainPanel.cpp
@@ -6,18 +6,7 @@
 #include "Views/List/ObjectMixerEditorList.h"
 #include "Views/MainPanel/ObjectMixerEditorMainPanel.h"
 
-#include "AssetRegistry/AssetRegistryModule.h"
-#include "ContentBrowserModule.h"
-#include "Framework/MultiBox/MultiBoxBuilder.h"
-#include "IContentBrowserSingleton.h"
-#include "ISinglePropertyView.h"
 #include "Kismet2/SClassPickerDialog.h"
-#include "PropertyEditorModule.h"
-#include "SPositiveActionButton.h"
-#include "Styling/StyleColors.h"
-
-#include "Widgets/Images/SImage.h"
-#include "Widgets/Input/SComboButton.h"
 #include "Widgets/Input/SMultiLineEditableTextBox.h"
 #include "Widgets/Layout/SSplitter.h"
 
@@ -32,26 +21,13 @@ void SObjectMixerEditorMainPanel::Construct(
 	
 	ChildSlot
 	[
-		SNew(SVerticalBox)
-			
-		+SVerticalBox::Slot()
-		.HAlign(HAlign_Fill)
-		.AutoHeight()
-		.Padding(FMargin(8.f, 0.f, 8.f, 0.f))
-		[
-			GeneratePanelToolbar()
-		]
+		SNew(SSplitter)
+		.Orientation(Orient_Vertical)
+		.ResizeMode(ESplitterResizeMode::FixedSize)
 
-		+SVerticalBox::Slot()
+		+SSplitter::Slot()
 		[
-			SNew(SSplitter)
-			.Orientation(Orient_Vertical)
-			.ResizeMode(ESplitterResizeMode::FixedSize)
-
-			+SSplitter::Slot()
-			[
-				MainPanel.Pin()->GetEditorList().Pin()->GetOrCreateWidget()
-			]
+			MainPanel.Pin()->GetEditorList().Pin()->GetOrCreateWidget()
 		]
 	];
 }
@@ -59,117 +35,6 @@ void SObjectMixerEditorMainPanel::Construct(
 SObjectMixerEditorMainPanel::~SObjectMixerEditorMainPanel()
 {
 	MainPanel.Reset();
-	ToolbarHBox.Reset();
-	ConcertButtonPtr.Reset();
-}
-
-TSharedRef<SWidget> SObjectMixerEditorMainPanel::GeneratePanelToolbar()
-{
-	FPropertyEditorModule& PropertyEditorModule = FModuleManager::LoadModuleChecked<FPropertyEditorModule>("PropertyEditor");
-
-	FSinglePropertyParams Params;
-	Params.NamePlacement = EPropertyNamePlacement::Hidden;
-	
-	return SAssignNew(ToolbarHBox, SHorizontalBox)
-				
-		// Add Console Variable button
-		+ SHorizontalBox::Slot()
-		.VAlign(VAlign_Fill)
-		.HAlign(HAlign_Left)
-		.Padding(FMargin(0, 4))
-		[
-			SNew(SPositiveActionButton)
-			.Text(LOCTEXT("AddObject", "Add"))
-			.OnClicked_Lambda([](){ return FReply::Handled(); })
-		]
-
-		// Filter Class Management Button
-		+ SHorizontalBox::Slot()
-		.VAlign(VAlign_Center)
-		.HAlign(HAlign_Right)
-		.Padding(FMargin(8.f, 4, 0.f, 4))
-		[
-			SNew(SComboButton)
-			.ToolTipText(LOCTEXT("FilterClassManagementButton_Tooltip", "Select a filter class"))
-			.ContentPadding(FMargin(4, 0.5f))
-			.ComboButtonStyle(&FAppStyle::Get().GetWidgetStyle<FComboButtonStyle>("ComboButton"))
-			.OnGetMenuContent(this, &SObjectMixerEditorMainPanel::OnGeneratePresetsMenu)
-			.ForegroundColor(FStyleColors::Foreground)
-			.ButtonContent()
-			[
-				SNew(SHorizontalBox)
-
-				+ SHorizontalBox::Slot()
-				.Padding(0, 1, 4, 0)
-				.AutoWidth()
-				[
-					SNew(SImage)
-					.Image(FAppStyle::Get().GetBrush("Icons.Filter"))
-					.ColorAndOpacity(FSlateColor::UseForeground())
-				]
-
-				+ SHorizontalBox::Slot()
-				.Padding(0, 1, 0, 0)
-				.AutoWidth()
-				[
-					SNew(STextBlock)
-					.Text(LOCTEXT("FilterClassToolbarButton", "Object Filter Class"))
-				]
-			]
-		];
-}
-
-TSharedRef<SWidget> SObjectMixerEditorMainPanel::OnGeneratePresetsMenu()
-{
-	FMenuBuilder MenuBuilder(true, NULL);
-
-	TArray<UClass*> DerivedClasses;
-	GetDerivedClasses(UObjectMixerObjectFilter::StaticClass(), DerivedClasses, true);
-
-	DerivedClasses.Remove(UObjectMixerObjectFilter::StaticClass());
-	DerivedClasses.Remove(UObjectMixerBlueprintObjectFilter::StaticClass());
-
-	if (DerivedClasses.Num())
-	{
-		const TSharedPtr<FObjectMixerEditorMainPanel, ESPMode::ThreadSafe> PinnedPanel = MainPanel.Pin();
-		check(PinnedPanel);
-		
-		MenuBuilder.BeginSection(NAME_None, LOCTEXT("SelectClassMenuSection", "Select Class"));
-		{
-			for (UClass* DerivedClass : DerivedClasses)
-			{
-				MenuBuilder.AddMenuEntry(
-				FText::FromName(DerivedClass->GetFName()),
-					FText::GetEmpty(),
-					FSlateIcon(),
-					FUIAction(
-						FExecuteAction::CreateSP(PinnedPanel.ToSharedRef(), &FObjectMixerEditorMainPanel::OnClassSelectionChanged, DerivedClass),
-						FCanExecuteAction::CreateLambda([](){ return true; }),
-						FIsActionChecked::CreateSP(PinnedPanel.ToSharedRef(), &FObjectMixerEditorMainPanel::IsClassSelected, DerivedClass)
-					),
-					NAME_None,
-					EUserInterfaceActionType::RadioButton
-				);
-			}
-		}
-		MenuBuilder.EndSection();
-	}
-	else
-	{
-		MenuBuilder.AddMenuEntry(LOCTEXT("NoFilterClassesAvailable", "No filter classes available."), FText::GetEmpty(), FSlateIcon(), FUIAction());
-	}
-
-	TSharedRef<SWidget> Widget = MenuBuilder.MakeWidget();
-	FChildren* ChildWidgets = Widget->GetChildren();
-	for (int32 ChildItr = 0; ChildItr < ChildWidgets->Num(); ChildItr++)
-	{
-		const TSharedRef<SWidget>& Child = ChildWidgets->GetChildAt(ChildItr);
-
-		Child->EnableToolTipForceField(false);
-	}
-	Widget->EnableToolTipForceField(false);
-	
-	return Widget;
 }
 
 #undef LOCTEXT_NAMESPACE
diff --git a/Engine/Source/Editor/ObjectMixer/ObjectMixer/Private/Views/Widgets/SObjectMixerPlacementAssetMenuEntry.cpp b/Engine/Source/Editor/ObjectMixer/ObjectMixer/Private/Views/Widgets/SObjectMixerPlacementAssetMenuEntry.cpp
new file mode 100644
--- /dev/null
+++ b/Engine/Source/Editor/ObjectMixer/ObjectMixer/Private/Views/Widgets/SObjectMixerPlacementAssetMenuEntry.cpp
@@ -0,0 +1,252 @@
+// Copyright Epic Games, Inc. All Rights Reserved.
+
+#include "SObjectMixerPlacementAssetMenuEntry.h"
+
+#include "AssetSelection.h"
+#include "ClassIconFinder.h"
+#include "DragAndDrop/AssetDragDropOp.h"
+#include "EditorClassUtils.h"
+#include "LevelEditorViewport.h"
+#include "PlacementMode/Public/IPlacementModeModule.h"
+#include "Styling/SlateIconFinder.h"
+#include "Styling/StyleColors.h"
+#include "Widgets/Layout/SSpacer.h"
+
+void SObjectMixerPlacementAssetMenuEntry::Construct(const FArguments& InArgs, const TSharedPtr<const FPlaceableItem>& InItem)
+{	
+	bIsPressed = false;
+
+	check(InItem.IsValid());
+
+	Item = InItem;
+
+	AssetImage = nullptr;
+
+	TSharedPtr< SHorizontalBox > ActorType = SNew( SHorizontalBox );
+
+	const bool bIsClass = Item->AssetData.GetClass() == UClass::StaticClass();
+	const bool bIsActor = bIsClass ? CastChecked<UClass>(Item->AssetData.GetAsset())->IsChildOf(AActor::StaticClass()) : false;
+
+	AActor* DefaultActor = nullptr;
+	if (Item->Factory != nullptr)
+	{
+		DefaultActor = Item->Factory->GetDefaultActor(Item->AssetData);
+	}
+	else if (bIsActor)
+	{
+		DefaultActor = CastChecked<AActor>(CastChecked<UClass>(Item->AssetData.GetAsset())->ClassDefaultObject);
+	}
+
+	UClass* DocClass = nullptr;
+	TSharedPtr<IToolTip> AssetEntryToolTip;
+	if(DefaultActor != nullptr)
+	{
+		DocClass = DefaultActor->GetClass();
+		AssetEntryToolTip = FEditorClassUtils::GetTooltip(DefaultActor->GetClass());
+	}
+
+	if (!AssetEntryToolTip.IsValid())
+	{
+		AssetEntryToolTip = FSlateApplicationBase::Get().MakeToolTip(Item->DisplayName);
+	}
+	
+	const FButtonStyle& ButtonStyle = FAppStyle::Get().GetWidgetStyle<FButtonStyle>( "Menu.Button" );
+	const float MenuIconSize = FAppStyle::Get().GetFloat("Menu.MenuIconSize");
+
+	Style = &ButtonStyle;
+
+	// Create doc link widget if there is a class to link to
+	TSharedRef<SWidget> DocWidget = SNew(SSpacer);
+	if(DocClass != NULL)
+	{
+		DocWidget = FEditorClassUtils::GetDocumentationLinkWidget(DocClass);
+		DocWidget->SetCursor( EMouseCursor::Default );
+	}
+
+	ChildSlot
+	.HAlign(HAlign_Fill)
+	.VAlign(VAlign_Fill)
+	[
+		SNew(SBorder)
+		.BorderImage( this, &SObjectMixerPlacementAssetMenuEntry::GetBorder )
+		.Cursor( EMouseCursor::GrabHand )
+		.ToolTip( AssetEntryToolTip )
+		.Padding(FMargin(5.f, 3.f, 5.f, 3.f))
+		[
+			SNew( SHorizontalBox )
+
+			+ SHorizontalBox::Slot()
+			.Padding(10.0f, 0.f, 10.f, 0.0f)
+			.AutoWidth()
+			.HAlign(HAlign_Left)
+			.VAlign(VAlign_Center)
+			[
+				SNew(SBox)
+				.WidthOverride(MenuIconSize)
+				.HeightOverride(MenuIconSize)
+				[
+					SNew(SImage)
+					.Image(this, &SObjectMixerPlacementAssetMenuEntry::GetIcon)
+					.ColorAndOpacity(FSlateColor::UseSubduedForeground())
+				]
+			]
+
+			+ SHorizontalBox::Slot()
+			.FillWidth(1.f)
+			.Padding(1.f, 0.f, 0.f, 0.f)
+			.VAlign(VAlign_Center)
+			.HAlign(HAlign_Left)
+			[
+
+				SNew( STextBlock )
+				.ColorAndOpacity(FSlateColor::UseForeground())
+				.Text( Item->DisplayName )
+			]
+
+
+			+ SHorizontalBox::Slot()
+			.VAlign(VAlign_Center)
+			.HAlign(HAlign_Right)
+			.AutoWidth()
+			[
+				SNew(SImage)
+                .ColorAndOpacity(FSlateColor::UseSubduedForeground())
+				.Image(FAppStyle::Get().GetBrush("Icons.DragHandle"))
+			]
+		]
+	];
+}
+
+
+const FSlateBrush* SObjectMixerPlacementAssetMenuEntry::GetIcon() const
+{
+
+	if (AssetImage != nullptr)
+	{
+		return AssetImage;
+	}
+
+	if (Item->ClassIconBrushOverride != NAME_None)
+	{
+		AssetImage = FSlateIconFinder::FindCustomIconBrushForClass(nullptr, TEXT("ClassIcon"), Item->ClassIconBrushOverride);
+	}
+	else
+	{
+		AssetImage = FSlateIconFinder::FindIconBrushForClass(FClassIconFinder::GetIconClassForAssetData(Item->AssetData));
+	}
+
+	return AssetImage;
+}
+
+
+FReply SObjectMixerPlacementAssetMenuEntry::OnMouseButtonDown(const FGeometry& MyGeometry, const FPointerEvent& MouseEvent)
+{
+	if ( MouseEvent.GetEffectingButton() == EKeys::LeftMouseButton )
+	{
+		bIsPressed = true;
+
+		return FReply::Handled().DetectDrag( SharedThis( this ), MouseEvent.GetEffectingButton() );
+	}
+
+	return FReply::Unhandled();
+}
+
+FReply SObjectMixerPlacementAssetMenuEntry::OnMouseButtonUp(const FGeometry& MyGeometry, const FPointerEvent& MouseEvent)
+{
+	if ( MouseEvent.GetEffectingButton() == EKeys::LeftMouseButton )
+	{
+		bIsPressed = false;
+
+		AActor* NewActor = nullptr;
+		UActorFactory* Factory = Item->Factory;
+		if (!Item->Factory)
+		{
+			// If no actor factory was found or failed, add the actor from the uclass
+			UClass* AssetClass = Item->AssetData.GetClass();
+			if (AssetClass)
+			{
+				UObject* ClassObject = AssetClass->GetDefaultObject();
+				FActorFactoryAssetProxy::GetFactoryForAssetObject(ClassObject);
+			}
+		}
+		NewActor = GEditor->GetEditorWorldContext().World()->SpawnActor(Factory->NewActorClass);
+		if (NewActor && GCurrentLevelEditingViewportClient)
+		{
+  			GEditor->MoveActorInFrontOfCamera(*NewActor, 
+  				GCurrentLevelEditingViewportClient->GetViewLocation(), 
+  				GCurrentLevelEditingViewportClient->GetViewRotation().Vector()
+  			);
+		}
+
+		if (!MouseEvent.IsControlDown())
+		{
+			FSlateApplication::Get().DismissAllMenus();
+		}
+
+		GEditor->SelectNone(true, true, false);
+		GEditor->SelectActor(NewActor, true, false);
+
+		return FReply::Handled();
+	}
+
+	return FReply::Unhandled();
+}
+
+FReply SObjectMixerPlacementAssetMenuEntry::OnDragDetected(const FGeometry& MyGeometry, const FPointerEvent& MouseEvent)
+{
+	bIsPressed = false;
+
+	if (FEditorDelegates::OnAssetDragStarted.IsBound())
+	{
+		TArray<FAssetData> DraggedAssetDatas;
+		DraggedAssetDatas.Add( Item->AssetData );
+		FEditorDelegates::OnAssetDragStarted.Broadcast( DraggedAssetDatas, Item->Factory );
+		return FReply::Handled();
+	}
+
+	if( MouseEvent.IsMouseButtonDown( EKeys::LeftMouseButton ) )
+	{
+		return FReply::Handled().BeginDragDrop(FAssetDragDropOp::New(Item->AssetData, Item->Factory));
+	}
+	else
+	{
+		return FReply::Handled();
+	}
+}
+
+bool SObjectMixerPlacementAssetMenuEntry::IsPressed() const
+{
+	return bIsPressed;
+}
+
+const FSlateBrush* SObjectMixerPlacementAssetMenuEntry::GetBorder() const
+{
+	if ( IsPressed() )
+	{
+		return &(Style->Pressed);
+	}
+	else if ( IsHovered() )
+	{
+		return &(Style->Hovered);
+	}
+	else
+	{
+		return &(Style->Normal);
+	}
+}
+
+FSlateColor SObjectMixerPlacementAssetMenuEntry::GetForegroundColor() const
+{
+	if (IsPressed())
+	{
+		return Style->PressedForeground;
+	}
+	else if (IsHovered())
+	{
+		return Style->HoveredForeground;
+	}
+	else
+	{
+		return Style->NormalForeground;
+	}
+}
diff --git a/Engine/Source/Editor/ObjectMixer/ObjectMixer/Private/Views/Widgets/SObjectMixerPlacementAssetMenuEntry.h b/Engine/Source/Editor/ObjectMixer/ObjectMixer/Private/Views/Widgets/SObjectMixerPlacementAssetMenuEntry.h
new file mode 100644
--- /dev/null
+++ b/Engine/Source/Editor/ObjectMixer/ObjectMixer/Private/Views/Widgets/SObjectMixerPlacementAssetMenuEntry.h
@@ -0,0 +1,36 @@
+// Copyright Epic Games, Inc. All Rights Reserved.
+
+#pragma once
+
+#include "Styling/SlateBrush.h"
+#include "Styling/SlateTypes.h"
+#include "Widgets/SCompoundWidget.h"
+
+class SObjectMixerPlacementAssetMenuEntry : public SCompoundWidget
+{
+public:
+	SLATE_BEGIN_ARGS(SObjectMixerPlacementAssetMenuEntry){}
+	SLATE_END_ARGS()
+
+	void Construct(const FArguments& InArgs, const TSharedPtr<const struct FPlaceableItem>& InItem);
+
+	virtual FReply OnMouseButtonDown(const FGeometry& MyGeometry, const FPointerEvent& MouseEvent) override;
+	virtual FReply OnMouseButtonUp(const FGeometry& MyGeometry, const FPointerEvent& MouseEvent) override;
+	virtual FReply OnDragDetected(const FGeometry& MyGeometry, const FPointerEvent& MouseEvent) override;
+
+	bool IsPressed() const;
+
+	TSharedPtr<const FPlaceableItem> Item;
+
+	virtual FSlateColor GetForegroundColor() const override;
+
+private:
+	const FSlateBrush* GetBorder() const;
+	const FSlateBrush* GetIcon() const;
+
+	bool bIsPressed = false;
+
+	const FButtonStyle* Style = nullptr;
+	
+	mutable const FSlateBrush* AssetImage = nullptr;
+};
diff --git a/Engine/Source/Editor/ObjectMixer/ObjectMixer/Public/ObjectFilter/ObjectMixerEditorObjectFilter.h b/Engine/Source/Editor/ObjectMixer/ObjectMixer/Public/ObjectFilter/ObjectMixerEditorObjectFilter.h
--- a/Engine/Source/Editor/ObjectMixer/ObjectMixer/Public/ObjectFilter/ObjectMixerEditorObjectFilter.h
+++ b/Engine/Source/Editor/ObjectMixer/ObjectMixer/Public/ObjectFilter/ObjectMixerEditorObjectFilter.h
@@ -8,7 +8,7 @@
 #include "ObjectMixerEditorObjectFilter.generated.h"
 
 UENUM(BlueprintType)
-enum class EObjectMixerPropertyInheritanceInclusionOptions : uint8
+enum class EObjectMixerInheritanceInclusionOptions : uint8
 {
 	None, // Get only the properties in the specified classes without considering parent or child classes + Specified Class
 	IncludeOnlyImmediateParent, // Get properties from the class the specified class immediately derives from, but not their parents + Specified Class
@@ -21,6 +21,15 @@ enum class EObjectMixerPropertyInheritanceInclusionOptions : uint8
 	IncludeOnlyImmediateParentAndAllChildren // IncludeOnlyImmediateParent + IncludeAllChildren + Specified Class
 };
 
+UENUM(BlueprintType)
+enum class EObjectMixerTreeViewMode : uint8
+{
+	Flat, // Show all matching objects in a flat list without any hierarchy
+	Folder, // Display objects in a folder hierarchy but without Object->SubObject hierarchy
+	ObjectSubObject, // Display objects as SubObjects of BaseObjects (usually actors) without considering folder hierarchy
+	FolderObjectSubObject // Display objects in a full hierarchy (folder->BaseObject->SubObject)
+};
+
 /**
  * Native class for filtering object types to Object Mixer.
  * Native C++ classes should inherit directly from this class.
@@ -34,52 +43,65 @@ public:
 
 	UObjectMixerObjectFilter() = default;
 	
-	/*
+	/**
 	 * Return the basic object types you want to filter for in your level.
 	 * For example, if you want to work with Lights, return ULightComponentBase.
 	 * If you also want to see the properties for parent or child classes,
 	 * override the GetObjectMixerPropertyInheritanceInclusionOptions and GetForceAddedColumns functions.
 	 */
-	virtual TArray<UClass*> GetObjectClassesToFilter() const { return {}; }
+	virtual TSet<UClass*> GetObjectClassesToFilter() const { return {}; }
+
+	/**
+	 * Return the basic actor types you want to be able to place using the Add button.
+	 * Note that only subclasses of AActor are supported and only those which have a registered factory.
+	 * This includes most engine actor types.
+	 */
+	virtual TSet<TSubclassOf<AActor>> GetObjectClassesToPlace() const { return {}; }
 
-	/*
+	/**
 	 * Get the text to display for the object name/label.
 	 * This is useful if one of your classes is a component type and you want the label of the component's owning actor, for example.
 	 * If not overridden, this returns the object's name.
+	 * @param InViewMode The tree view display option (i.e. Flat list or hierarchy).
 	 */
-	virtual FText GetRowDisplayName(UObject* InObject) const;
+	virtual FText GetRowDisplayName(
+		UObject* InObject, EObjectMixerTreeViewMode InViewMode) const;
 
-	/*
+	/**
 	 * Controls how to display the row's visibility icon. Return true if the object should be visible.
 	 * Generally this should work like the Scene Outliner does.
 	 * If not overridden, we use the Editor Visibility of the object's AActor outer (unless it's an actor itself).
+	 * @param InViewMode The tree view display option (i.e. Flat list or hierarchy) (unused)
 	 */
-	virtual bool GetRowEditorVisibility(UObject* InObject) const;
+	virtual bool GetRowEditorVisibility(
+		UObject* InObject, EObjectMixerTreeViewMode InViewMode) const;
 
-	/*
+	/**
 	 * Controls what happens when the row's visibility icon is clicked.
 	 * Generally this should work like the Scene Outliner does.
 	 * If not overridden, we set the Editor Visibility of the object's AActor outer (unless it's an actor itself).
+	 * @param InViewMode The tree view display option (i.e. Flat list or hierarchy) (unused)
 	 */
-	virtual void OnSetRowEditorVisibility(UObject* InObject, bool bNewIsVisible) const;
+	virtual void OnSetRowEditorVisibility(
+		UObject* InObject, bool bNewIsVisible, EObjectMixerTreeViewMode InViewMode) const;
 
-	/*
+	/**
 	 * Specify a list of property names corresponding to columns you want to show by default.
 	 * For example, you can specify "Intensity" and "LightColor" to show only those property columns by default in the UI.
 	 * Columns not specified will not be shown by default but can be enabled by the user in the UI.
 	 */
-	virtual TArray<FName> GetColumnsToShowByDefault() const;
+	virtual TSet<FName> GetColumnsToShowByDefault() const;
 
-	/*
+	/**
 	 * Specify a list of property names corresponding to columns you don't want to ever show.
 	 * For example, you can specify "Intensity" and "LightColor" to ensure that they can't be enabled or shown in the UI.
 	 * Columns not specified can be enabled by the user in the UI.
 	 */
-	virtual TArray<FName> GetColumnsFilter() const;
+	virtual TSet<FName> GetColumnsToExclude() const;
 
-	/*
+	/**
 	 * Specify a list of property names found in parent classes you want to show that aren't in the specified classes.
-	 * Note that properties specified here do not override the properties specified in GetColumnsFilter(),
+	 * Note that properties specified here do not override the properties specified in GetColumnsToExclude(),
 	 * but do override the supported property tests so these will appear even if ShouldIncludeUnsupportedProperties returns false.
 	 * For example, a ULightComponent displays "LightColor" in the editor's details panel,
 	 * but ULightComponent itself doesn't have a property named "LightColor". Instead it's in its parent class, ULightComponentBase.
@@ -87,36 +109,42 @@ public:
 	 * Specify "LightColor" in this function to ensure that "LightColor" will appear as a column as long as
 	 * the property is accessible to one of the specified classes regardless of which parent class it comes from.
 	 */
-	virtual TArray<FName> GetForceAddedColumns() const { return {}; }
+	virtual TSet<FName> GetForceAddedColumns() const { return {}; }
 
-	/*
+	/**
 	 * Specify whether we should return only the properties of the specified classes or the properties of parent and child classes.
 	 * Defaults to 'None' which only considers the properties of the specified classes.
 	 * If you're not seeing all the properties you expected, try overloading this function.
 	 */
-	virtual EObjectMixerPropertyInheritanceInclusionOptions GetObjectMixerPropertyInheritanceInclusionOptions() const;
+	virtual EObjectMixerInheritanceInclusionOptions GetObjectMixerPropertyInheritanceInclusionOptions() const;
 
-	/*
+	/**
+	 * Specify whether we should return only the specified classes or the parent and child classes in placement mode.
+	 * Defaults to 'None' which only considers the specified classes.
+	 */
+	virtual EObjectMixerInheritanceInclusionOptions GetObjectMixerPlacementClassInclusionOptions() const;
+
+	/**
 	 * If true, properties that are not visible in the details panel and properties not supported by SSingleProperty will be selectable.
 	 * Defaults to false.
 	 */
 	virtual bool ShouldIncludeUnsupportedProperties() const;
 
-	static TArray<UClass*> GetParentAndChildClassesFromSpecifiedClasses(
-		const TArray<UClass*>& InSpecifiedClasses, EObjectMixerPropertyInheritanceInclusionOptions Options);
+	static TSet<UClass*> GetParentAndChildClassesFromSpecifiedClasses(
+		const TSet<UClass*>& InSpecifiedClasses, EObjectMixerInheritanceInclusionOptions Options);
 
-	static EFieldIterationFlags GetDesiredFieldIterationFlags(const bool bIncludeInheritedProperties);
+	static TSet<UClass*> GetParentAndChildClassesFromSpecifiedClasses(
+		const TSet<TSubclassOf<AActor>>& InSpecifiedClasses, EObjectMixerInheritanceInclusionOptions Options);
 
-	// Folder view should reflect outliner folder structure if possible
-	// Mute/solo lights
+	static EFieldIterationFlags GetDesiredFieldIterationFlags(const bool bIncludeInheritedProperties);
 
 protected:
 
-	/*
+	/**
 	 * Given a set of property names you wish to include, returns a list of all other properties on InObject not found in ExcludeList.
 	 * Useful when defining default visible columns in a list view.
 	 */
-	TArray<FName> GenerateIncludeListFromExcludeList(const TSet<FName>& ExcludeList) const;
+	TSet<FName> GenerateIncludeListFromExcludeList(const TSet<FName>& ExcludeList) const;
 };
 
 /**
@@ -129,88 +157,107 @@ class OBJECTMIXEREDITOR_API UObjectMixerBlueprintObjectFilter : public UObjectMi
 	GENERATED_BODY()
 public:
 	
-	/*
+	/**
 	 * Return the basic object types you want to filter for in your level.
 	 * For example, if you want to work with Lights, return ULightComponentBase.
 	 * If you also want to see the properties for parent or child classes,
 	 * override the GetObjectMixerPropertyInheritanceInclusionOptions and GetForceAddedColumns functions.
 	 */
 	UFUNCTION(BlueprintNativeEvent, BlueprintCallable, Category = "Object Mixer")
-	TArray<UClass*> GetObjectClassesToFilter() const override;
+	TSet<UClass*> GetObjectClassesToFilter() const override;
 
-	TArray<UClass*> GetObjectClassesToFilter_Implementation() const
+	TSet<UClass*> GetObjectClassesToFilter_Implementation() const
 	{
 		return Super::GetObjectClassesToFilter();
 	}
 
-	/*
+	/**
+	 * Return the basic actor types you want to be able to place using the Add button.
+	 * Note that only subclasses of AActor are supported and only those which have a registered factory.
+	 * This includes most engine actor types.
+	 */
+	UFUNCTION(BlueprintNativeEvent, BlueprintCallable, Category = "Object Mixer")
+	TSet<TSubclassOf<AActor>> GetObjectClassesToPlace() const override;
+
+	TSet<TSubclassOf<AActor>> GetObjectClassesToPlace_Implementation() const
+	{
+		return Super::GetObjectClassesToPlace();
+	}
+
+	/**
 	 * Get the text to display for the object name/label.
 	 * This is useful if one of your classes is a component type and you want the label of the component's owning actor, for example.
 	 * If not overridden, this returns the object's name.
+	 * @param InViewMode The tree view display option (i.e. Flat list or hierarchy).
 	 */
 	UFUNCTION(BlueprintNativeEvent, BlueprintCallable, Category = "Object Mixer")
-	FText GetRowDisplayName(UObject* InObject) const override;
+	FText GetRowDisplayName(
+		UObject* InObject, EObjectMixerTreeViewMode InViewMode = EObjectMixerTreeViewMode::FolderObjectSubObject) const override;
 	
-	FText GetRowDisplayName_Implementation(UObject* InObject) const
+	FText GetRowDisplayName_Implementation(UObject* InObject, EObjectMixerTreeViewMode InViewMode) const
 	{
-		return Super::GetRowDisplayName(InObject);
+		return Super::GetRowDisplayName(InObject, InViewMode);
 	}
 	
-	/*
+	/**
 	 * Controls how to display the row's visibility icon. Return true if the object should be visible.
 	 * Generally this should work like the Scene Outliner does.
 	 * If not overridden, we use the Editor Visibility of the object's AActor outer (unless it's an actor itself).
+	 * @param InViewMode The tree view display option (i.e. Flat list or hierarchy) (unused)
 	 */
 	UFUNCTION(BlueprintNativeEvent, BlueprintCallable, Category = "Object Mixer")
-	bool GetRowEditorVisibility(UObject* InObject) const override;
+	bool GetRowEditorVisibility(
+		UObject* InObject, EObjectMixerTreeViewMode InViewMode = EObjectMixerTreeViewMode::FolderObjectSubObject) const override;
 
-	bool GetRowEditorVisibility_Implementation(UObject* InObject) const
+	bool GetRowEditorVisibility_Implementation(UObject* InObject, EObjectMixerTreeViewMode InViewMode) const
 	{
-		return Super::GetRowEditorVisibility(InObject);
+		return Super::GetRowEditorVisibility(InObject, InViewMode);
 	}
 	
-	/*
+	/**
 	 * Controls what happens when the row's visibility icon is clicked.
 	 * Generally this should work like the Scene Outliner does.
 	 * If not overridden, we set the Editor Visibility of the object's AActor outer (unless it's an actor itself).
+	 * @param InViewMode The tree view display option (i.e. Flat list or hierarchy) (unused)
 	 */
 	UFUNCTION(BlueprintNativeEvent, BlueprintCallable, Category = "Object Mixer")
-	void OnSetRowEditorVisibility(UObject* InObject, bool bNewIsVisible) const override;
+	void OnSetRowEditorVisibility(
+		UObject* InObject, bool bNewIsVisible, EObjectMixerTreeViewMode InViewMode = EObjectMixerTreeViewMode::FolderObjectSubObject) const override;
 
-	void OnSetRowEditorVisibility_Implementation(UObject* InObject, bool bNewIsVisible) const
+	void OnSetRowEditorVisibility_Implementation(UObject* InObject, bool bNewIsVisible, EObjectMixerTreeViewMode InViewMode) const
 	{
-		return Super::OnSetRowEditorVisibility(InObject, bNewIsVisible);
+		return Super::OnSetRowEditorVisibility(InObject, bNewIsVisible, InViewMode);
 	}
 	
-	/*
+	/**
 	 * Specify a list of property names corresponding to columns you want to show by default.
 	 * For example, you can specify "Intensity" and "LightColor" to show only those property columns by default in the UI.
 	 * Columns not specified will not be shown by default but can be enabled by the user in the UI.
 	 */
 	UFUNCTION(BlueprintNativeEvent, BlueprintCallable, Category = "Object Mixer")
-	TArray<FName> GetColumnsToShowByDefault() const override;
+	TSet<FName> GetColumnsToShowByDefault() const override;
 	
-	TArray<FName> GetColumnsToShowByDefault_Implementation() const
+	TSet<FName> GetColumnsToShowByDefault_Implementation() const
 	{
 		return Super::GetColumnsToShowByDefault();
 	}
 	
-	/*
+	/**
 	 * Specify a list of property names corresponding to columns you don't want to ever show.
 	 * For example, you can specify "Intensity" and "LightColor" to ensure that they can't be enabled or shown in the UI.
 	 * Columns not specified can be enabled by the user in the UI.
 	 */
 	UFUNCTION(BlueprintNativeEvent, BlueprintCallable, Category = "Object Mixer")
-	TArray<FName> GetColumnsFilter() const override;
+	TSet<FName> GetColumnsToExclude() const override;
 
-	TArray<FName> GetColumnsFilter_Implementation() const
+	TSet<FName> GetColumnsToExclude_Implementation() const
 	{
-		return Super::GetColumnsFilter();
+		return Super::GetColumnsToExclude();
 	}
 
-	/*
+	/**
 	 * Specify a list of property names found in parent classes you want to show that aren't in the specified classes.
-	 * Note that properties specified here do not override the properties specified in GetColumnsFilter().
+	 * Note that properties specified here do not override the properties specified in GetColumnsToExclude().
 	 * For example, a ULightComponent displays "LightColor" in the editor's details panel,
 	 * but ULightComponent itself doesn't have a property named "LightColor". Instead it's in its parent class, ULightComponentBase.
 	 * In this scenario, ULightComponent is specified and PropertyInheritanceInclusionOptions is None, so "LightColor" won't appear by default.
@@ -218,27 +265,39 @@ public:
 	 * the property is accessible to one of the specified classes regardless of which parent class it comes from.
 	 */
 	UFUNCTION(BlueprintNativeEvent, BlueprintCallable, Category = "Object Mixer")
-	TArray<FName> GetForceAddedColumns() const override;
+	TSet<FName> GetForceAddedColumns() const override;
 	
-	TArray<FName> GetForceAddedColumns_Implementation() const
+	TSet<FName> GetForceAddedColumns_Implementation() const
 	{
 		return Super::GetForceAddedColumns();
 	}
 	
-	/*
+	/**
 	 * Specify whether we should return only the properties of the specified classes or the properties of parent and child classes.
 	 * Defaults to 'None' which only considers the properties of the specified classes.
 	 * If you're not seeing all the properties you expected, try overloading this function.
 	 */
 	UFUNCTION(BlueprintNativeEvent, BlueprintCallable, Category = "Object Mixer")
-	EObjectMixerPropertyInheritanceInclusionOptions GetObjectMixerPropertyInheritanceInclusionOptions() const override;
+	EObjectMixerInheritanceInclusionOptions GetObjectMixerPropertyInheritanceInclusionOptions() const override;
 
-	EObjectMixerPropertyInheritanceInclusionOptions GetObjectMixerPropertyInheritanceInclusionOptions_Implementation() const
+	EObjectMixerInheritanceInclusionOptions GetObjectMixerPropertyInheritanceInclusionOptions_Implementation() const
 	{
 		return Super::GetObjectMixerPropertyInheritanceInclusionOptions();
 	}
 
-	/*
+	/**
+	 * Specify whether we should return only the specified classes or the parent and child classes in placement mode.
+	 * Defaults to 'None' which only considers the specified classes.
+	 */
+	UFUNCTION(BlueprintNativeEvent, BlueprintCallable, Category = "Object Mixer")
+	EObjectMixerInheritanceInclusionOptions GetObjectMixerPlacementClassInclusionOptions() const override;
+
+	EObjectMixerInheritanceInclusionOptions GetObjectMixerPlacementClassInclusionOptions_Implementation() const
+	{
+		return Super::GetObjectMixerPlacementClassInclusionOptions();
+	}
+
+	/**
 	 * If true, properties that are not visible in the details panel and properties not supported by SSingleProperty will be selectable.
 	 * Defaults to false.
 	 */
diff --git a/Engine/Source/Editor/ObjectMixer/ObjectMixer/Public/ObjectMixerEditorModule.h b/Engine/Source/Editor/ObjectMixer/ObjectMixer/Public/ObjectMixerEditorModule.h
--- a/Engine/Source/Editor/ObjectMixer/ObjectMixer/Public/ObjectMixerEditorModule.h
+++ b/Engine/Source/Editor/ObjectMixer/ObjectMixer/Public/ObjectMixerEditorModule.h
@@ -26,7 +26,7 @@ public:
 	
 	virtual TSharedPtr<SWidget> MakeObjectMixerDialog() const;
 
-	/*
+	/**
 	 * Regenerate the list items and refresh the list. Call when adding or removing variables.
 	 */
 	virtual void RequestRebuildList() const;
@@ -60,6 +60,8 @@ public:
 protected:
 
 	virtual TSharedRef<SDockTab> SpawnMainPanelTab();
+
+	virtual void BindDelegates();
 	
 	/** Lives for as long as the module is loaded. */
 	TSharedPtr<FObjectMixerEditorMainPanel> MainPanel;
diff --git a/Engine/Source/Editor/ObjectMixer/ObjectMixer/Public/ObjectMixerEditorProjectSettings.h b/Engine/Source/Editor/ObjectMixer/ObjectMixer/Public/ObjectMixerEditorProjectSettings.h
--- a/Engine/Source/Editor/ObjectMixer/ObjectMixer/Public/ObjectMixerEditorProjectSettings.h
+++ b/Engine/Source/Editor/ObjectMixer/ObjectMixer/Public/ObjectMixerEditorProjectSettings.h
@@ -2,8 +2,6 @@
 
 #pragma once
 
-#include "ObjectFilter/ObjectMixerEditorObjectFilter.h"
-
 #include "ObjectMixerEditorProjectSettings.generated.h"
 
 UCLASS(config = Engine, defaultconfig)
@@ -19,5 +17,5 @@ public:
 	 * If false, a new object will be created every time the filter object is accessed.
 	 */
 	UPROPERTY(Config, EditAnywhere, Category="Object Mixer")
-	bool bShouldCacheFilterObject = true;
+	bool bExpandTreeViewItemsByDefault = true;
 };
diff --git a/Engine/Source/Editor/ObjectMixer/ObjectMixer/Public/Views/List/ObjectMixerEditorList.h b/Engine/Source/Editor/ObjectMixer/ObjectMixer/Public/Views/List/ObjectMixerEditorList.h
--- a/Engine/Source/Editor/ObjectMixer/ObjectMixer/Public/Views/List/ObjectMixerEditorList.h
+++ b/Engine/Source/Editor/ObjectMixer/ObjectMixer/Public/Views/List/ObjectMixerEditorList.h
@@ -27,32 +27,69 @@ public:
 
 	void CacheObjectFilterObject();
 
-	/*
-	 * Returns either ObjectClassFromWhichToGeneratePropertyColumns or ObjectClassOverride, if set.
-	 * @param bForceNoOverride If true, always return ObjectClassFromWhichToGeneratePropertyColumns and not ObjectClassOverride.
+	/**
+	 * Get the style of the tree (flat list or hierarchy)
+	 */
+	EObjectMixerTreeViewMode GetTreeViewMode()
+	{
+		return TreeViewMode;
+	}
+	/**
+	 * Set the style of the tree (flat list or hierarchy)
+	 */
+	void SetTreeViewMode(EObjectMixerTreeViewMode InViewMode)
+	{
+		TreeViewMode = InViewMode;
+		RequestRebuildList();
+	}
+
+	/**
+	 * Returns result from Filter->GetObjectClassesToFilter.
+	 */
+	TSet<UClass*> GetObjectClassesToFilter()
+	{
+		if (const UObjectMixerObjectFilter* Filter = GetObjectFilter())
+		{
+			return Filter->GetObjectClassesToFilter();
+		}
+		
+		return {};
+	}
+
+	/**
+	 * Returns result from Filter->GetObjectClassesToPlace.
 	 */
-	TArray<UClass*> GetObjectClasses(bool bForceNoOverride = false)
+	TSet<TSubclassOf<AActor>> GetObjectClassesToPlace()
 	{
-		TArray<UClass*> ReturnValue;
+		TSet<TSubclassOf<AActor>> ReturnValue;
 
 		if (const UObjectMixerObjectFilter* Filter = GetObjectFilter())
 		{
-			ReturnValue = Filter->GetObjectClassesToFilter();
+			ReturnValue = Filter->GetObjectClassesToPlace();
 		}
 		
 		return ReturnValue;
 	}
 
+	/**
+	 * Get the row that has solo visibility. All other rows should be set to temporarily invisible in editor.
+	 */
 	TWeakPtr<FObjectMixerEditorListRow> GetSoloRow()
 	{
 		return SoloRow;
 	}
 
+	/**
+	 * Set the row that has solo visibility. This does not set temporary editor invisibility for other rows.
+	 */
 	void SetSoloRow(TSharedRef<FObjectMixerEditorListRow> InRow)
 	{
 		SoloRow = InRow;
 	}
 
+	/**
+	 * Clear the row that has solo visibility. This does not remove temporary editor invisibility for other rows.
+	 */
 	void ClearSoloRow()
 	{
 		SoloRow = nullptr;
@@ -62,27 +99,32 @@ public:
 
 	void ClearList() const;
 
-	/*
+	/**
 	 * Regenerate the list items and refresh the list. Call when adding or removing variables.
 	 */
 	void RequestRebuildList() const;
 
 	/**
 	 * Refresh filters and sorting.
-	 * Useful for when the list state has gone stale but the variable count has not changed.
+	 * Useful for when the list state has gone stale but the item count has not changed.
 	 */
 	void RefreshList() const;
 
+	bool IsClassSelected(UClass* InNewClass) const;
+
 	TSubclassOf<UObjectMixerObjectFilter> GetObjectFilterClass() const
 	{
 		return ObjectFilterClass;
 	}
 
-	void SetObjectFilterClass(TSubclassOf<UObjectMixerObjectFilter> InObjectFilterClass)
+	void SetObjectFilterClass(UClass* InObjectFilterClass)
 	{
-		ObjectFilterClass = InObjectFilterClass;
-		CacheObjectFilterObject();
-		RequestRebuildList();
+		if (ensureAlwaysMsgf(InObjectFilterClass->IsChildOf(UObjectMixerObjectFilter::StaticClass()), TEXT("%hs: Class '%s' is not a child of UObjectMixerObjectFilter."), __FUNCTION__, *InObjectFilterClass->GetName()))
+		{
+			ObjectFilterClass = InObjectFilterClass;
+			CacheObjectFilterObject();
+			RequestRebuildList();
+		}
 	}
 
 private:
@@ -92,6 +134,11 @@ private:
 	 */
 	TSubclassOf<UObjectMixerObjectFilter> ObjectFilterClass;
 
+	/**
+	 * Determines the style of the tree (flat list or hierarchy)
+	 */
+	EObjectMixerTreeViewMode TreeViewMode = EObjectMixerTreeViewMode::FolderObjectSubObject;
+
 	TSharedPtr<SObjectMixerEditorList> ListWidget;
 
 	TStrongObjectPtr<UObjectMixerObjectFilter> ObjectFilterPtr;
diff --git a/Engine/Source/Editor/ObjectMixer/ObjectMixer/Public/Views/List/ObjectMixerEditorListRow.h b/Engine/Source/Editor/ObjectMixer/ObjectMixer/Public/Views/List/ObjectMixerEditorListRow.h
--- a/Engine/Source/Editor/ObjectMixer/ObjectMixer/Public/Views/List/ObjectMixerEditorListRow.h
+++ b/Engine/Source/Editor/ObjectMixer/ObjectMixer/Public/Views/List/ObjectMixerEditorListRow.h
@@ -5,6 +5,7 @@
 #include "Containers/Array.h"
 #include "Containers/UnrealString.h"
 #include "Layout/Visibility.h"
+#include "ObjectFilter/ObjectMixerEditorObjectFilter.h"
 #include "Templates/SharedPointer.h"
 #include "UObject/Object.h"
 #include "UObject/WeakObjectPtr.h"
@@ -19,11 +20,12 @@ typedef TSharedPtr<FObjectMixerEditorListRow> FObjectMixerEditorListRowPtr;
 
 struct OBJECTMIXEREDITOR_API FObjectMixerEditorListRow final : TSharedFromThis<FObjectMixerEditorListRow>
 {
-	enum EObjectMixerEditorListRowType
+	enum EObjectMixerEditorListRowType : uint8
 	{
-		None,
-		Group,
-		SingleItem
+		None = 0,
+		Folder, // Usually an Outliner folder
+		ContainerObject, // Usually an actor that contains a matching subobject or is the attach parent of a matching actor
+		MatchingObject // The object that has the properties we wish to affect
 	};
 
 	~FObjectMixerEditorListRow();
@@ -32,15 +34,14 @@ struct OBJECTMIXEREDITOR_API FObjectMixerEditorListRow final : TSharedFromThis<F
 	
 	FObjectMixerEditorListRow(
 		const TWeakObjectPtr<UObject> InObject, const EObjectMixerEditorListRowType InRowType, 
-		const TSharedRef<SObjectMixerEditorList>& InListView, 
-		const TWeakPtr<FObjectMixerEditorListRow>& InDirectParentRow)
+		const TSharedRef<SObjectMixerEditorList>& InListView, const FText& InDisplayNameOverride = FText::GetEmpty())
 	: ObjectRef(InObject)
 	, RowType(InRowType)
 	, ListViewPtr(InListView)
-	, DirectParentRow(InDirectParentRow)
+	, DisplayNameOverride(InDisplayNameOverride)
 	{}
 
-	[[nodiscard]] UObject* GetObject()
+	[[nodiscard]] UObject* GetObject() const
 	{
 		if (ObjectRef.IsValid())
 		{
@@ -63,33 +64,28 @@ struct OBJECTMIXEREDITOR_API FObjectMixerEditorListRow final : TSharedFromThis<F
 	TWeakPtr<FObjectMixerEditorListRow> GetDirectParentRow() const;
 	void SetDirectParentRow(const TWeakPtr<FObjectMixerEditorListRow>& InDirectParentRow);
 	
-	/* bHasGeneratedChildren must be true to get actual children. */
+	/** bHasGeneratedChildren must be true to get actual children. */
 	[[nodiscard]] const TArray<FObjectMixerEditorListRowPtr>& GetChildRows() const;
-	/* bHasGeneratedChildren must be true to get an accurate value. */
+	/** bHasGeneratedChildren must be true to get an accurate value. */
 	[[nodiscard]] int32 GetChildCount() const;
 	void SetChildRows(const TArray<FObjectMixerEditorListRowPtr>& InChildRows);
 	void AddToChildRows(const FObjectMixerEditorListRowPtr& InRow);
 	void InsertChildRowAtIndex(const FObjectMixerEditorListRowPtr& InRow, const int32 AtIndex = 0);
 
-	[[nodiscard]] bool GetIsTreeViewItemExpanded() const;
+	[[nodiscard]] bool GetIsTreeViewItemExpanded();
 	void SetIsTreeViewItemExpanded(const bool bNewExpanded);
-	
-	[[nodiscard]] bool GetShouldFlashOnScrollIntoView() const;
-	void SetShouldFlashOnScrollIntoView(const bool bNewShouldFlashOnScrollIntoView);
 
 	[[nodiscard]] bool GetShouldExpandAllChildren() const;
 	void SetShouldExpandAllChildren(const bool bNewShouldExpandAllChildren);
 
-	void ResetToStartupValueAndSource() const;
-
-	/*
+	/**
 	 *Individual members of InTokens will be considered "AnyOf" or "OR" searches. If SearchTerms contains any individual member it will match.
 	 *Members will be tested for a space character (" "). If a space is found, a subsearch will be run.
 	 *This subsearch will be an "AllOf" or "AND" type search in which all strings, separated by a space, must be found in the search terms.
 	 */
 	bool MatchSearchTokensToSearchTerms(const TArray<FString> InTokens, ESearchCase::Type InSearchCase = ESearchCase::IgnoreCase);
 
-	/* This overload creates tokens from a string first, then calls ExecuteSearchOnChildNodes(const TArray<FString>& Tokens). */
+	/** This overload creates tokens from a string first, then calls ExecuteSearchOnChildNodes(const TArray<FString>& Tokens). */
 	void ExecuteSearchOnChildNodes(const FString& SearchString) const;
 	void ExecuteSearchOnChildNodes(const TArray<FString>& Tokens) const;
 
@@ -98,25 +94,46 @@ struct OBJECTMIXEREDITOR_API FObjectMixerEditorListRow final : TSharedFromThis<F
 
 	[[nodiscard]] bool GetIsSelected();
 
-	[[nodiscard]] bool ShouldBeVisible() const;
-	[[nodiscard]] EVisibility GetDesiredVisibility() const;
+	[[nodiscard]] bool ShouldRowWidgetBeVisible() const;
+	[[nodiscard]] EVisibility GetDesiredRowWidgetVisibility() const;
 
-	[[nodiscard]] bool HasVisibleChildren() const
+	[[nodiscard]] bool HasVisibleChildRowWidgets() const
 	{
 		return false;
 	}
 
+	[[nodiscard]] FText GetDisplayName();
+
+	[[nodiscard]] const FText& GetDisplayNameOverride()
+	{
+		return DisplayNameOverride;
+	}
+
+	void SetDisplayNameOverride(const FText& InDisplayNameOverride)
+	{
+		DisplayNameOverride = InDisplayNameOverride;
+	}
+
 	[[nodiscard]] TWeakPtr<SObjectMixerEditorList> GetListViewPtr() const
 	{
 		return ListViewPtr;
 	}
 
+	/**
+	 * Determines the style of the tree (flat list or hierarchy)
+	 */
+	EObjectMixerTreeViewMode GetTreeViewMode();
+	/**
+	 * Determine the style of the tree (flat list or hierarchy)
+	 */
+	void SetTreeViewMode(EObjectMixerTreeViewMode InViewMode);
+
 	[[nodiscard]] TArray<FObjectMixerEditorListRowPtr> GetSelectedTreeViewItems() const;
 
 	const FSlateBrush* GetObjectIconBrush();
 
 	bool GetObjectVisibility();
-	void SetObjectVisibility(const bool bNewIsVisible);
+	void SetObjectVisibility(const bool bNewIsVisible, const bool bIsRecursive = false);
 
 	bool IsThisRowSolo() const;
 	void SetThisAsSoloRow();
@@ -125,14 +142,17 @@ struct OBJECTMIXEREDITOR_API FObjectMixerEditorListRow final : TSharedFromThis<F
 
 private:
 
+	FObjectMixerEditorListRowPtr GetAsShared();
+
 	TWeakObjectPtr<UObject> ObjectRef;
-	EObjectMixerEditorListRowType RowType = SingleItem;
+	EObjectMixerEditorListRowType RowType = MatchingObject;
 	TArray<FObjectMixerEditorListRowPtr> ChildRows;
 	
 	TWeakPtr<SObjectMixerEditorList> ListViewPtr;
+
+	FText DisplayNameOverride;
 	
 	bool bIsTreeViewItemExpanded = false;
-	bool bShouldFlashOnScrollIntoView = false;
 
 	int32 ChildDepth = 0;
 
diff --git a/Engine/Source/Editor/ObjectMixer/ObjectMixer/Public/Views/List/SObjectMixerEditorList.h b/Engine/Source/Editor/ObjectMixer/ObjectMixer/Public/Views/List/SObjectMixerEditorList.h
--- a/Engine/Source/Editor/ObjectMixer/ObjectMixer/Public/Views/List/SObjectMixerEditorList.h
+++ b/Engine/Source/Editor/ObjectMixer/ObjectMixer/Public/Views/List/SObjectMixerEditorList.h
@@ -93,7 +93,7 @@ public:
 	 */
 	void RefreshList();
 
-	/*
+	/**
 	 * Regenerate the list items and refresh the list. Call when adding or removing items.
 	 */
 	void RequestRebuildList(const FString& InItemToScrollTo = "");
@@ -101,7 +101,7 @@ public:
 	[[nodiscard]] TArray<FObjectMixerEditorListRowPtr> GetSelectedTreeViewItems() const;
 	int32 GetSelectedTreeViewItemCount() const;
 
-	void SetSelectedTreeViewItemActorsEditorVisible(const bool bNewIsVisible);
+	void SetSelectedTreeViewItemActorsEditorVisible(const bool bNewIsVisible, const bool bIsRecursive = false);
 
 	bool IsTreeViewItemSelected(TSharedRef<FObjectMixerEditorListRow> Item);
 
@@ -134,8 +134,18 @@ public:
 
 	bool DoesTreeViewHaveVisibleChildren() const;
 
+	bool IsTreeViewItemExpanded(const TSharedPtr<FObjectMixerEditorListRow>& Row) const;
 	void SetTreeViewItemExpanded(const TSharedPtr<FObjectMixerEditorListRow>& RowToExpand, const bool bNewExpansion) const;
 
+	/**
+	 * Determines the style of the tree (flat list or hierarchy)
+	 */
+	EObjectMixerTreeViewMode GetTreeViewMode();
+	/**
+	 * Determine the style of the tree (flat list or hierarchy)
+	 */
+	void SetTreeViewMode(EObjectMixerTreeViewMode InViewMode);
+
 	void ToggleFilterActive(const FString& FilterName);
 	void EvaluateIfRowsPassFilters(const bool bShouldRefreshAfterward = true);
 
@@ -164,6 +174,14 @@ public:
 	static const FName EditorVisibilityColumnName;
 	static const FName EditorVisibilitySoloColumnName;
 
+	inline static TFunction<bool(const FObjectMixerEditorListRowPtr&, const FObjectMixerEditorListRowPtr&)> SortByTypeThenName =
+		[](const FObjectMixerEditorListRowPtr& A, const FObjectMixerEditorListRowPtr& B)
+	{
+		return (A->GetRowType() < B->GetRowType() ||
+			(A->GetRowType() == B->GetRowType() && A->GetDisplayName().ToString() < B->GetDisplayName().ToString())
+		);
+	};
+
 private:
 
 	virtual void Tick(const FGeometry& AllottedGeometry, const double InCurrentTime, const float InDeltaTime) override;
@@ -178,22 +196,25 @@ private:
 	TSharedPtr<SHeaderRow> HeaderRow;
 	TSharedRef<SWidget> GenerateHeaderRowContextMenu() const;
 
-	bool bShouldRebuild = false;
+	TSharedRef<SWidget> GenerateToolbar();
+	TSharedRef<SWidget> OnGenerateAddObjectButtonMenu() const;
 
-	/*
+	bool bShouldRebuild = false;
+	
+	/**
 	 * Regenerate the list items and refresh the list. Call when adding or removing items.
 	 */
 	void RebuildList();
 
-	/*
+	/**
 	 * Only adds properties that pass a series of tests, including having only one unique entry in the column list array.
 	 * @param bForceIncludeProperty If true, only Skiplist and Uniqueness tests will be checked, bypassing class, blueprint editability and other requirements.
-	 * @param
+	 * @param PropertySkipList These property names will be skipped when they are encountered in the iteration.
 	 */
 	bool AddUniquePropertyColumnsToHeaderRow(
 		FProperty* Property,
 		const bool bForceIncludeProperty = false,
-		const TArray<FName>& PropertySkipList = {}
+		const TSet<FName>& PropertySkipList = {}
 	);
 	void AddBuiltinColumnsToHeaderRow();
 	TSharedPtr<SHeaderRow> GenerateHeaderRow();
@@ -201,6 +222,7 @@ private:
 
 	void SetupFilters();
 
+	TSharedRef<SWidget> OnGenerateFilterClassMenu();
 	TSharedRef<SWidget> BuildShowOptionsMenu();
 	
 	void FlushMemory(const bool bShouldKeepMemoryAllocated);
@@ -218,6 +240,10 @@ private:
 
 	//  Tree View Implementation
 
+	/** Saves tree item expanded states to be recalled after the tree view is regenerated. */
+	void CacheTreeState();
+	void RestoreTreeState(const bool bFlushCache = true);
+	void BuildPerformanceCacheAndGenerateHeaderIfNeeded();
 	void GenerateTreeView();
 	void FindVisibleTreeViewObjects();
 	void FindVisibleObjectsAndRequestTreeRefresh();
@@ -231,6 +257,8 @@ private:
 
 	TSharedPtr<STreeView<FObjectMixerEditorListRowPtr>> TreeViewPtr;
 
+	TMap<FString, bool> TreeItemExpansionStateCache;
+
 	/** All Tree view objects */
 	TArray<FObjectMixerEditorListRowPtr> TreeViewRootObjects;
 	/** Visible Tree view objects, after filters */
@@ -249,6 +277,11 @@ private:
 			return A->GetSortOrder() < B->GetSortOrder();
 		};
 
-	FDelegateHandle OnActorSpawnedHandle;
-	FDelegateHandle OnActorDestroyedHandle;
+	// Performance cache
+	TSet<UClass*> ObjectClassesToFilterCache;
+	TSet<FName> ColumnsToShowByDefaultCache;
+	TSet<FName> ColumnsToExcludeCache;
+	TSet<FName> ForceAddedColumnsCache;
+	EObjectMixerInheritanceInclusionOptions PropertyInheritanceInclusionOptionsCache = EObjectMixerInheritanceInclusionOptions::None;
+	bool bShouldIncludeUnsupportedPropertiesCache = false;
 };
diff --git a/Engine/Source/Editor/ObjectMixer/ObjectMixer/Public/Views/List/SObjectMixerEditorListRow.h b/Engine/Source/Editor/ObjectMixer/ObjectMixer/Public/Views/List/SObjectMixerEditorListRow.h
--- a/Engine/Source/Editor/ObjectMixer/ObjectMixer/Public/Views/List/SObjectMixerEditorListRow.h
+++ b/Engine/Source/Editor/ObjectMixer/ObjectMixer/Public/Views/List/SObjectMixerEditorListRow.h
@@ -37,8 +37,6 @@ public:
 	// TOptional<EItemDropZone> HandleCanAcceptDrop(const FDragDropEvent& DragDropEvent, EItemDropZone DropZone, FObjectMixerEditorListRowPtr TargetItem);
 	// FReply HandleAcceptDrop(const FDragDropEvent& DragDropEvent, EItemDropZone DropZone, FObjectMixerEditorListRowPtr TargetItem);
 	
-	void FlashRow();
-
 private:
 
 	bool IsVisible() const;
@@ -48,33 +46,15 @@ private:
 	
 	/** Get the brush for this widget */
 	const FSlateBrush* GetVisibilityBrush() const;
-	
-	EVisibility GetFlashImageVisibility() const;
-	FSlateColor GetFlashImageColorAndOpacity() const;
 
 	static const FSlateBrush* GetBorderImage(const FObjectMixerEditorListRow::EObjectMixerEditorListRowType InRowType);
 
-	TSharedRef<SWidget> GenerateCells(const FName& InColumnName, const TSharedPtr<FObjectMixerEditorListRow> PinnedItem);
+	TSharedPtr<SWidget> GenerateCells(const FName& InColumnName, const TSharedPtr<FObjectMixerEditorListRow> PinnedItem);
 
 	void OnPropertyChanged(const FProperty* Property, void* ContainerWithChangedProperty);
 	
 	TWeakPtr<FObjectMixerEditorListRow> Item;
 	
-	TSharedPtr<IToolTip> HoverToolTip;
-
-	TArray<TSharedPtr<SImage>> FlashImages;
-
-	TSet<FDelegateHandle> StructureChangeDelegateHandles;
-
-	TSharedPtr<SObjectMixerEditorListValueInput> ValueChildInputWidget;
-	
-	TSharedPtr<SObjectMixerEditorListRowHoverWidgets> HoverableWidgetsPtr;
-
-	FCurveSequence FlashAnimation;
-
-	const float FlashAnimationDuration = 0.75f;
-	const FLinearColor FlashColor = FLinearColor::White;
-	
 	const FSlateBrush* VisibleHoveredBrush = nullptr;
 	const FSlateBrush* VisibleNotHoveredBrush = nullptr;
 	const FSlateBrush* NotVisibleHoveredBrush = nullptr;
diff --git a/Engine/Source/Editor/ObjectMixer/ObjectMixer/Public/Views/MainPanel/ObjectMixerEditorMainPanel.h b/Engine/Source/Editor/ObjectMixer/ObjectMixer/Public/Views/MainPanel/ObjectMixerEditorMainPanel.h
--- a/Engine/Source/Editor/ObjectMixer/ObjectMixer/Public/Views/MainPanel/ObjectMixerEditorMainPanel.h
+++ b/Engine/Source/Editor/ObjectMixer/ObjectMixer/Public/Views/MainPanel/ObjectMixerEditorMainPanel.h
@@ -17,7 +17,7 @@ public:
 
 	TSharedRef<SWidget> GetOrCreateWidget();
 
-	/*
+	/**
 	 * Regenerate the list items and refresh the list. Call when adding or removing items.
 	 */
 	void RequestRebuildList() const;
diff --git a/Engine/Source/Editor/ObjectMixer/ObjectMixer/Public/Views/MainPanel/SObjectMixerEditorMainPanel.h b/Engine/Source/Editor/ObjectMixer/ObjectMixer/Public/Views/MainPanel/SObjectMixerEditorMainPanel.h
--- a/Engine/Source/Editor/ObjectMixer/ObjectMixer/Public/Views/MainPanel/SObjectMixerEditorMainPanel.h
+++ b/Engine/Source/Editor/ObjectMixer/ObjectMixer/Public/Views/MainPanel/SObjectMixerEditorMainPanel.h
@@ -2,13 +2,9 @@
 
 #pragma once
 
-#include "Widgets/Input/SCheckBox.h"
 #include "Widgets/SCompoundWidget.h"
 
 class FObjectMixerEditorMainPanel;
-class SBox;
-class SHorizontalBox;
-class SVerticalBox;
 
 class OBJECTMIXEREDITOR_API SObjectMixerEditorMainPanel final : public SCompoundWidget
 {
@@ -28,14 +24,4 @@ private:
 
 	/** A reference to the struct that controls this widget */
 	TWeakPtr<FObjectMixerEditorMainPanel> MainPanel;
-
-	TSharedPtr<SHorizontalBox> ToolbarHBox;
-
-	/** A reference to the button which opens the plugin settings */
-	TSharedPtr<SCheckBox> ConcertButtonPtr;
-
-	/** Creates the toolbar at the top of the MainPanel widget */
-	TSharedRef<SWidget> GeneratePanelToolbar();
-
-	TSharedRef<SWidget> OnGeneratePresetsMenu();
 };
