Adding additional includes in preparation for removing a global include

#rnx

diff --git a/Engine/Source/Editor/ObjectMixer/ObjectMixer/Private/ObjectMixerEditorModule.cpp b/Engine/Source/Editor/ObjectMixer/ObjectMixer/Private/ObjectMixerEditorModule.cpp
--- a/Engine/Source/Editor/ObjectMixer/ObjectMixer/Private/ObjectMixerEditorModule.cpp
+++ b/Engine/Source/Editor/ObjectMixer/ObjectMixer/Private/ObjectMixerEditorModule.cpp
@@ -8,6 +8,7 @@
 #include "Views/MainPanel/ObjectMixerEditorMainPanel.h"
 
 #include "Editor.h"
+#include "Misc/TransactionObjectEvent.h"
 #include "Framework/Docking/TabManager.h"
 #include "ISettingsModule.h"
 #include "Modules/ModuleManager.h"
