Object Mixer

-Fix for //UE5/Main - Static Analysis Win64 (MSVC) - warning: Dereferencing NULL pointer 'BaseActor'

diff --git a/Engine/Source/Editor/ObjectMixer/ObjectMixer/Private/Views/List/SObjectMixerEditorList.cpp b/Engine/Source/Editor/ObjectMixer/ObjectMixer/Private/Views/List/SObjectMixerEditorList.cpp
--- a/Engine/Source/Editor/ObjectMixer/ObjectMixer/Private/Views/List/SObjectMixerEditorList.cpp
+++ b/Engine/Source/Editor/ObjectMixer/ObjectMixer/Private/Views/List/SObjectMixerEditorList.cpp
@@ -1293,23 +1293,26 @@ void SObjectMixerEditorList::GenerateTreeView()
 						{
 							TSharedPtr<FObjectMixerEditorListRow> OwningActorRow;
 
-							if (const TSharedRef<FObjectMixerEditorListRow>* Match = CreatedObjectMap.Find(BaseActor))
+							if (BaseActor)
 							{
-								OwningActorRow = *Match;
-							}
-							else
-							{
-								OwningActorRow = MakeShared<FObjectMixerEditorListRow>(
-									BaseActor, FObjectMixerEditorListRow::ContainerObject, SharedThis(this));
+								if (const TSharedRef<FObjectMixerEditorListRow>* Match = CreatedObjectMap.Find(BaseActor))
+								{
+									OwningActorRow = *Match;
+								}
+								else
+								{
+									OwningActorRow = MakeShared<FObjectMixerEditorListRow>(
+										BaseActor, FObjectMixerEditorListRow::ContainerObject, SharedThis(this));
 
-								CreatedObjectMap.Add(BaseActor, OwningActorRow.ToSharedRef());
-							}
+									CreatedObjectMap.Add(BaseActor, OwningActorRow.ToSharedRef());
+								}
 
-							if (OwningActorRow)
-							{
-								OwningActorRow->AddToChildRows(TopLevelRow);
+								if (OwningActorRow)
+								{
+									OwningActorRow->AddToChildRows(TopLevelRow);
 													
-								TopLevelRow = OwningActorRow.ToSharedRef();
+									TopLevelRow = OwningActorRow.ToSharedRef();
+								}
 							}
 						}
 					}
@@ -1351,34 +1354,37 @@ void SObjectMixerEditorList::GenerateTreeView()
 					if (ListModelPtr.Pin()->GetTreeViewMode() == EObjectMixerTreeViewMode::Folder ||
 						ListModelPtr.Pin()->GetTreeViewMode() == EObjectMixerTreeViewMode::FolderObjectSubObject)
 					{
-						FFolder BaseActorFolder = BaseActor->GetFolder();
-
-						while (!BaseActorFolder.IsNone())
+						if (BaseActor)
 						{
-							TSharedPtr<FObjectMixerEditorListRow> FolderRow;
-							
-							if (const TSharedRef<FObjectMixerEditorListRow>* Match = FolderMap.Find(BaseActorFolder.GetPath()))
-							{
-								FolderRow = *Match;
-							}
-							else
+							FFolder BaseActorFolder = BaseActor->GetFolder();
+
+							while (!BaseActorFolder.IsNone())
 							{
-								FolderRow =
-									MakeShared<FObjectMixerEditorListRow>(
-										nullptr, FObjectMixerEditorListRow::Folder, SharedThis(this),
-										FText::FromName(BaseActorFolder.GetLeafName()));
+								TSharedPtr<FObjectMixerEditorListRow> FolderRow;
+							
+								if (const TSharedRef<FObjectMixerEditorListRow>* Match = FolderMap.Find(BaseActorFolder.GetPath()))
+								{
+									FolderRow = *Match;
+								}
+								else
+								{
+									FolderRow =
+										MakeShared<FObjectMixerEditorListRow>(
+											nullptr, FObjectMixerEditorListRow::Folder, SharedThis(this),
+											FText::FromName(BaseActorFolder.GetLeafName()));
 
-								FolderMap.Add(BaseActorFolder.GetPath(), FolderRow.ToSharedRef());
-							}
+									FolderMap.Add(BaseActorFolder.GetPath(), FolderRow.ToSharedRef());
+								}
 
-							if (FolderRow)
-							{
-								FolderRow->AddToChildRows(TopLevelRow);
+								if (FolderRow)
+								{
+									FolderRow->AddToChildRows(TopLevelRow);
 													
-								TopLevelRow = FolderRow.ToSharedRef();
-							}
+									TopLevelRow = FolderRow.ToSharedRef();
+								}
 
-							BaseActorFolder = BaseActorFolder.GetParent();
+								BaseActorFolder = BaseActorFolder.GetParent();
+							}
 						}
 					}
 				}
