Object Mixer

-Get categories from properties in header context menu

diff --git a/Engine/Source/Editor/ObjectMixer/ObjectMixer/Private/Views/List/SObjectMixerEditorList.cpp b/Engine/Source/Editor/ObjectMixer/ObjectMixer/Private/Views/List/SObjectMixerEditorList.cpp
--- a/Engine/Source/Editor/ObjectMixer/ObjectMixer/Private/Views/List/SObjectMixerEditorList.cpp
+++ b/Engine/Source/Editor/ObjectMixer/ObjectMixer/Private/Views/List/SObjectMixerEditorList.cpp
@@ -626,12 +626,18 @@ bool SObjectMixerEditorList::AddUniquePropertyColumnsToHeaderRow(
 				})
 			)
 		{
+			FString CategoryName = "Generated Properties";
+			if (const FString* CategoryMeta = Property->FindMetaData("Category"))
+			{
+				CategoryName = *CategoryMeta;
+			}
+			
 			ListViewColumns.Add(
 				{
 					Property, PropertyName,
 					Property->GetDisplayNameText(),
 					EListViewColumnType::PropertyGenerated,
-					"Generated Properties",
+					*CategoryName,
 					true, true, false
 				}
 			);
