Object Mixer

-Implement category buttons
-Save category assignments to Project Settings
-Create Context Menu interface for category assignment (right click on any row to add or assign category to object, then select category filter above)
-Move generic data and functions from list model to mainpanel model to support multiple lists
-Fix: Header context menu should follow the style of outliner

diff --git a/Engine/Source/Editor/ObjectMixer/LightMixer/Private/LightMixerModule.cpp b/Engine/Source/Editor/ObjectMixer/LightMixer/Private/LightMixerModule.cpp
--- a/Engine/Source/Editor/ObjectMixer/LightMixer/Private/LightMixerModule.cpp
+++ b/Engine/Source/Editor/ObjectMixer/LightMixer/Private/LightMixerModule.cpp
@@ -40,6 +40,11 @@ void FLightMixerModule::Initialize()
 	DefaultFilterClass = ULightMixerObjectFilter::StaticClass();	
 }
 
+FName FLightMixerModule::GetModuleName()
+{
+	return "LightMixer";
+}
+
 void FLightMixerModule::SetupMenuItemVariables()
 {
 	TabLabel = LOCTEXT("LightMixerTabLabel", "Light Mixer");
diff --git a/Engine/Source/Editor/ObjectMixer/LightMixer/Public/LightMixerModule.h b/Engine/Source/Editor/ObjectMixer/LightMixer/Public/LightMixerModule.h
--- a/Engine/Source/Editor/ObjectMixer/LightMixer/Public/LightMixerModule.h
+++ b/Engine/Source/Editor/ObjectMixer/LightMixer/Public/LightMixerModule.h
@@ -17,6 +17,7 @@ public:
 	
 	//~ Begin FObjectMixerEditorModule overrides
 	virtual void Initialize() override;
+	virtual FName GetModuleName() override;
 	virtual void SetupMenuItemVariables() override;
 	virtual FName GetTabSpawnerId() override;
 	virtual void RegisterProjectSettings() const override;
diff --git a/Engine/Source/Editor/ObjectMixer/ObjectMixer/ObjectMixerEditor.build.cs b/Engine/Source/Editor/ObjectMixer/ObjectMixer/ObjectMixerEditor.build.cs
--- a/Engine/Source/Editor/ObjectMixer/ObjectMixer/ObjectMixerEditor.build.cs
+++ b/Engine/Source/Editor/ObjectMixer/ObjectMixer/ObjectMixerEditor.build.cs
@@ -34,6 +34,7 @@ public class ObjectMixerEditor : ModuleRules
 				"InputCore",
 				"Kismet",
 				"Projects",
+				"SceneOutliner",
 				"Slate",
 				"SlateCore",
 				"ToolMenus", 
diff --git a/Engine/Source/Editor/ObjectMixer/ObjectMixer/Private/ObjectMixerEditorModule.cpp b/Engine/Source/Editor/ObjectMixer/ObjectMixer/Private/ObjectMixerEditorModule.cpp
--- a/Engine/Source/Editor/ObjectMixer/ObjectMixer/Private/ObjectMixerEditorModule.cpp
+++ b/Engine/Source/Editor/ObjectMixer/ObjectMixer/Private/ObjectMixerEditorModule.cpp
@@ -8,14 +8,16 @@
 #include "Views/MainPanel/ObjectMixerEditorMainPanel.h"
 
 #include "Editor.h"
-#include "Misc/TransactionObjectEvent.h"
 #include "Framework/Docking/TabManager.h"
 #include "ISettingsModule.h"
+#include "Misc/TransactionObjectEvent.h"
 #include "Modules/ModuleManager.h"
 #include "ToolMenus.h"
 #include "WorkspaceMenuStructure.h"
 #include "WorkspaceMenuStructureModule.h"
 
+const FName FObjectMixerEditorModule::BaseObjectMixerModuleName("ObjectMixerEditor");
+
 #define LOCTEXT_NAMESPACE "FObjectMixerEditorModule"
 
 IMPLEMENT_MODULE(FObjectMixerEditorModule, ObjectMixerEditor)
@@ -96,6 +98,11 @@ FObjectMixerEditorModule& FObjectMixerEditorModule::Get()
 	return FModuleManager::LoadModuleChecked< FObjectMixerEditorModule >("ObjectMixerEditor");
 }
 
+FName FObjectMixerEditorModule::GetModuleName()
+{
+	return "ObjectMixerEditor";
+}
+
 TSharedPtr<SWidget> FObjectMixerEditorModule::MakeObjectMixerDialog() const
 {
 	if (MainPanel.IsValid())
@@ -231,7 +238,8 @@ TSharedPtr<FWorkspaceItem> FObjectMixerEditorModule::GetWorkspaceGroup()
 
 TSharedRef<SDockTab> FObjectMixerEditorModule::SpawnMainPanelTab()
 {
-	MainPanel = MakeShared<FObjectMixerEditorMainPanel>();
+	MainPanel = MakeShared<FObjectMixerEditorMainPanel>(GetModuleName());
+	MainPanel->Init();
 	
 	const TSharedRef<SDockTab> DockTab =
 		SNew(SDockTab)
diff --git a/Engine/Source/Editor/ObjectMixer/ObjectMixer/Private/Views/List/ObjectMixerEditorList.cpp b/Engine/Source/Editor/ObjectMixer/ObjectMixer/Private/Views/List/ObjectMixerEditorList.cpp
--- a/Engine/Source/Editor/ObjectMixer/ObjectMixer/Private/Views/List/ObjectMixerEditorList.cpp
+++ b/Engine/Source/Editor/ObjectMixer/ObjectMixer/Private/Views/List/ObjectMixerEditorList.cpp
@@ -4,16 +4,9 @@
 
 #include "Views/List/SObjectMixerEditorList.h"
 
-#include "UObject/UObjectGlobals.h"
-
 FObjectMixerEditorList::~FObjectMixerEditorList()
 {
-	ListWidget.Reset();
-
-	if (ObjectFilterPtr.IsValid())
-	{
-		ObjectFilterPtr.Reset();
-	}
+	FlushWidget();
 }
 
 void FObjectMixerEditorList::FlushWidget()
@@ -33,34 +26,6 @@ TSharedRef<SWidget> FObjectMixerEditorList::GetOrCreateWidget()
 	return ListWidget.ToSharedRef();
 }
 
-UObjectMixerObjectFilter* FObjectMixerEditorList::GetObjectFilter()
-{
-	if (!ObjectFilterPtr.IsValid())
-	{
-		CacheObjectFilterObject();
-	}
-
-	return ObjectFilterPtr.Get();
-}
-
-void FObjectMixerEditorList::CacheObjectFilterObject()
-{
-	if (ObjectFilterPtr.IsValid())
-	{
-		ObjectFilterPtr.Reset();
-	}
-	
-	if (const UClass* Class = GetObjectFilterClass())
-	{
-		ObjectFilterPtr = TStrongObjectPtr(NewObject<UObjectMixerObjectFilter>(GetTransientPackage(), Class));
-	}
-}
-
-void FObjectMixerEditorList::SetSearchString(const FString& SearchString)
-{
-	ListWidget->SetSearchStringInSearchInputField(SearchString);
-}
-
 void FObjectMixerEditorList::ClearList() const
 {
 	if (ListWidget.IsValid())
@@ -85,7 +50,33 @@ void FObjectMixerEditorList::RefreshList() const
 	}
 }
 
-bool FObjectMixerEditorList::IsClassSelected(UClass* InNewClass) const
+void FObjectMixerEditorList::ExecuteListViewSearchOnAllRows(const FString& SearchString,
+	const bool bShouldRefreshAfterward)
+{
+	ListWidget->ExecuteListViewSearchOnAllRows(SearchString, bShouldRefreshAfterward);
+}
+
+void FObjectMixerEditorList::EvaluateIfRowsPassFilters(const bool bShouldRefreshAfterward) const
+{
+	ListWidget->EvaluateIfRowsPassFilters();
+}
+
+TWeakPtr<FObjectMixerEditorMainPanel> FObjectMixerEditorList::GetMainPanelModel()
+{
+	return MainPanelModelPtr;
+}
+
+TWeakPtr<FObjectMixerEditorListRow> FObjectMixerEditorList::GetSoloRow()
+{
+	return GetMainPanelModel().Pin()->GetSoloRow();
+}
+
+void FObjectMixerEditorList::SetSoloRow(TSharedRef<FObjectMixerEditorListRow> InRow)
+{
+	GetMainPanelModel().Pin()->SetSoloRow(InRow);
+}
+
+void FObjectMixerEditorList::ClearSoloRow()
 {
-	return InNewClass == GetObjectFilterClass();
+	GetMainPanelModel().Pin()->ClearSoloRow();
 }
diff --git a/Engine/Source/Editor/ObjectMixer/ObjectMixer/Private/Views/List/ObjectMixerEditorListRow.cpp b/Engine/Source/Editor/ObjectMixer/ObjectMixer/Private/Views/List/ObjectMixerEditorListRow.cpp
--- a/Engine/Source/Editor/ObjectMixer/ObjectMixer/Private/Views/List/ObjectMixerEditorListRow.cpp
+++ b/Engine/Source/Editor/ObjectMixer/ObjectMixer/Private/Views/List/ObjectMixerEditorListRow.cpp
@@ -28,9 +28,12 @@ UObjectMixerObjectFilter* FObjectMixerEditorListRow::GetObjectFilter() const
 	{
 		if (const TSharedPtr<FObjectMixerEditorList> PinnedListModel = PinnedListView->GetListModelPtr().Pin())
 		{
-			if (UObjectMixerObjectFilter* Filter = PinnedListModel->GetObjectFilter())
+			if (const TSharedPtr<FObjectMixerEditorMainPanel> PinnedMainPanel = PinnedListModel->GetMainPanelModel().Pin())
 			{
-				return Filter;
+				if (UObjectMixerObjectFilter* Filter = PinnedMainPanel->GetObjectFilter())
+				{
+					return Filter;
+				}
 			}
 		}
 	}
@@ -38,19 +41,37 @@ UObjectMixerObjectFilter* FObjectMixerEditorListRow::GetObjectFilter() const
 	return nullptr;
 }
 
-FObjectMixerEditorListRow::EObjectMixerEditorListRowType FObjectMixerEditorListRow::GetRowType() const
+bool FObjectMixerEditorListRow::IsObjectRefInSelectedCategories() const
 {
-	return RowType;
-}
+	check (ListViewPtr.IsValid());
+	
+	if (RowType != EObjectMixerEditorListRowType::None && RowType != EObjectMixerEditorListRowType::Folder)
+	{
+		if (const TSharedPtr<FObjectMixerEditorList> ListModel = ListViewPtr.Pin()->GetListModelPtr().Pin())
+		{
+			if (const TSharedPtr<FObjectMixerEditorMainPanel> MainPanel = ListModel->GetMainPanelModel().Pin())
+			{
+				const TSet<FName>& CategorySelection = MainPanel->GetCurrentCategorySelection();
 
-int32 FObjectMixerEditorListRow::GetChildDepth() const
-{
-	return ChildDepth;
+				if (CategorySelection.Num() == 0)
+				{
+					return true;
+				}
+
+				const TSet<FName>& ObjectAssignedCategories = MainPanel->GetCategoriesForObject(GetObject());
+
+				const TSet<FName> Intersection = CategorySelection.Intersect(ObjectAssignedCategories);
+
+				return Intersection.Num() > 0;
+			}
+		}
+	}
+	return false;
 }
 
-void FObjectMixerEditorListRow::SetChildDepth(const int32 InDepth)
+FObjectMixerEditorListRow::EObjectMixerEditorListRowType FObjectMixerEditorListRow::GetRowType() const
 {
-	ChildDepth = InDepth;
+	return RowType;
 }
 
 int32 FObjectMixerEditorListRow::GetSortOrder() const
@@ -127,41 +148,41 @@ bool FObjectMixerEditorListRow::MatchSearchTokensToSearchTerms(
 {
 	// If the search is cleared we'll consider the row to pass search
 	bool bMatchFound = true;
-
-	if (const TObjectPtr<UObject> Object = GetObject())
+	
+	if (CachedSearchTerms.IsEmpty())
 	{
-		if (CachedSearchTerms.IsEmpty())
-		{
-			CachedSearchTerms = Object->GetName();
+		CachedSearchTerms = GetDisplayNameOverride().ToString() + " ";
 
+		if (const TObjectPtr<UObject> Object = GetObject())
+		{			
 			if (const UObjectMixerObjectFilter* Filter = GetObjectFilter())
 			{
-				CachedSearchTerms = Filter->GetRowDisplayName(Object, GetTreeViewMode()).ToString();
+				CachedSearchTerms += Filter->GetRowDisplayName(Object, GetTreeViewMode()).ToString();
 			}
 		}
+	}
 
-		// Match any
-		for (const FString& Token : InTokens)
+	// Match any
+	for (const FString& Token : InTokens)
+	{
+		// Match all of these
+		const FString SpaceDelimiter = " ";
+		TArray<FString> OutSpacedArray;
+		if (Token.Contains(SpaceDelimiter) && Token.ParseIntoArray(OutSpacedArray, *SpaceDelimiter, true) > 1)
 		{
-			// Match all of these
-			const FString SpaceDelimiter = " ";
-			TArray<FString> OutSpacedArray;
-			if (Token.Contains(SpaceDelimiter) && Token.ParseIntoArray(OutSpacedArray, *SpaceDelimiter, true) > 1)
-			{
-				bMatchFound = Algo::AllOf(OutSpacedArray, [this, InSearchCase](const FString& Comparator)
-				{
-					return CachedSearchTerms.Contains(Comparator, InSearchCase);
-				});
-			}
-			else
+			bMatchFound = Algo::AllOf(OutSpacedArray, [this, InSearchCase](const FString& Comparator)
 			{
-				bMatchFound = CachedSearchTerms.Contains(Token, InSearchCase);
-			}
+				return CachedSearchTerms.Contains(Comparator, InSearchCase);
+			});
+		}
+		else
+		{
+			bMatchFound = CachedSearchTerms.Contains(Token, InSearchCase);
+		}
 
-			if (bMatchFound)
-			{
-				break;
-			}
+		if (bMatchFound)
+		{
+			break;
 		}
 	}
 
@@ -188,9 +209,11 @@ void FObjectMixerEditorListRow::ExecuteSearchOnChildNodes(const TArray<FString>&
 			continue;
 		}
 
-		if (ChildRow->GetRowType() == EObjectMixerEditorListRowType::Folder)
+		const bool bMatch = ChildRow->MatchSearchTokensToSearchTerms(Tokens);;
+
+		if (ChildRow->GetChildCount() > 0)
 		{
-			if (ChildRow->MatchSearchTokensToSearchTerms(Tokens))
+			if (bMatch)
 			{
 				// If the group name matches then we pass an empty string to search child nodes since we want them all to be visible
 				ChildRow->ExecuteSearchOnChildNodes("");
@@ -201,10 +224,6 @@ void FObjectMixerEditorListRow::ExecuteSearchOnChildNodes(const TArray<FString>&
 				ChildRow->ExecuteSearchOnChildNodes(Tokens);
 			}
 		}
-		else
-		{
-			ChildRow->MatchSearchTokensToSearchTerms(Tokens);
-		}
 	}
 }
 
@@ -235,6 +254,24 @@ EVisibility FObjectMixerEditorListRow::GetDesiredRowWidgetVisibility() const
 	return ShouldRowWidgetBeVisible() ? EVisibility::Visible : EVisibility::Collapsed;
 }
 
+bool FObjectMixerEditorListRow::HasVisibleChildRowWidgets() const
+{
+	if (ChildRows.Num() == 0)
+	{
+		return false;
+	}
+
+	for (const auto& ChildRow : ChildRows)
+	{
+		if (ChildRow->ShouldRowWidgetBeVisible())
+		{
+			return true;
+		}
+	}
+
+	return false;
+}
+
 FText FObjectMixerEditorListRow::GetDisplayName()
 {
 	if (!GetDisplayNameOverride().IsEmpty())
@@ -261,14 +298,6 @@ EObjectMixerTreeViewMode FObjectMixerEditorListRow::GetTreeViewMode()
 	return PinnedListModel->GetTreeViewMode();
 }
 
-void FObjectMixerEditorListRow::SetTreeViewMode(EObjectMixerTreeViewMode InViewMode)
-{
-	if (TSharedPtr<SObjectMixerEditorList> PinnedListModel = GetListViewPtr().Pin())
-	{
-		PinnedListModel->SetTreeViewMode(InViewMode);
-	}
-}
-
 TArray<FObjectMixerEditorListRowPtr> FObjectMixerEditorListRow::GetSelectedTreeViewItems() const
 {
 	return ListViewPtr.Pin()->GetSelectedTreeViewItems();
diff --git a/Engine/Source/Editor/ObjectMixer/ObjectMixer/Private/Views/List/SObjectMixerEditorList.cpp b/Engine/Source/Editor/ObjectMixer/ObjectMixer/Private/Views/List/SObjectMixerEditorList.cpp
--- a/Engine/Source/Editor/ObjectMixer/ObjectMixer/Private/Views/List/SObjectMixerEditorList.cpp
+++ b/Engine/Source/Editor/ObjectMixer/ObjectMixer/Private/Views/List/SObjectMixerEditorList.cpp
@@ -3,7 +3,7 @@
 #include "Views/List/SObjectMixerEditorList.h"
 
 // View Filters
-#include "Views/List/ObjectMixerEditorListFilters/ObjectMixerEditorListFilter_Source.h"
+#include "Views/List/ObjectMixerEditorListFilters/IObjectMixerEditorListFilter.h"
 
 #include "ObjectFilter/ObjectMixerEditorObjectFilter.h"
 #include "ObjectMixerEditorLog.h"
@@ -11,21 +11,24 @@
 #include "ObjectMixerEditorProjectSettings.h"
 #include "ObjectMixerEditorStyle.h"
 #include "Views/List/SObjectMixerEditorListRow.h"
-#include "Views/Widgets/SObjectMixerPlacementAssetMenuEntry.h"
 
+#include "ActorBrowsingMode.h"
 #include "Algo/AllOf.h"
 #include "Algo/AnyOf.h"
 #include "Editor.h"
+#include "SceneOutlinerMenuContext.h"
+#include "Selection.h"
 #include "Framework/MultiBox/MultiBoxBuilder.h"
-#include "SPositiveActionButton.h"
 #include "DragAndDrop/AssetDragDropOp.h"
+#include "Editor/UnrealEdEngine.h"
 #include "PlacementMode/Public/IPlacementModeModule.h"
 #include "Styling/StyleColors.h"
 #include "UObject/UObjectIterator.h"
-#include "Widgets/Input/SSearchBox.h"
+#include "Views/Widgets/ObjectMixerEditorListMenuContext.h"
 #include "Widgets/Layout/SScrollBox.h"
 #include "Widgets/Layout/SWidgetSwitcher.h"
 #include "Widgets/Text/SRichTextBlock.h"
+#include "WorldPartition/WorldPartition.h"
 
 #define LOCTEXT_NAMESPACE "ObjectMixerEditorList"
 
@@ -44,20 +47,10 @@ void SObjectMixerEditorList::Construct(const FArguments& InArgs, TSharedRef<FObj
 				.CanSelectGeneratedColumn(false)
 				.Visibility(EVisibility::Visible)
 				;
-
-	SetupFilters();
 	
 	ChildSlot
 	[
 		SNew(SVerticalBox)
-		
-		+ SVerticalBox::Slot()
-		.VAlign(VAlign_Top)
-		.AutoHeight()
-		.Padding(FMargin(8.f, 0.f, 8.f, 0.f))
-		[
-			GenerateToolbar()
-		]
 
 		+ SVerticalBox::Slot()
 		.HAlign(HAlign_Fill)
@@ -77,6 +70,18 @@ void SObjectMixerEditorList::Construct(const FArguments& InArgs, TSharedRef<FObj
 				SAssignNew(TreeViewPtr, STreeView<FObjectMixerEditorListRowPtr>)
 				.HeaderRow(HeaderRow)
 				.SelectionMode(ESelectionMode::Multi)
+				.OnContextMenuOpening_Lambda([this]()
+				{
+					TSharedPtr<SWidget> Widget;
+					
+					if (GetListModelPtr().IsValid())
+					{
+						Widget = UObjectMixerEditorListMenuContext::CreateContextMenu(
+							{GetSelectedTreeViewItems(), GetListModelPtr().Pin()->GetMainPanelModel()});
+					}
+					
+					return Widget;
+				})
 				.OnSelectionChanged_Lambda([this] (const FObjectMixerEditorListRowPtr& Row, const ESelectInfo::Type SelectionType)
 				{
 					if (GEditor && FSlateApplication::Get().GetModifierKeys().IsAltDown())
@@ -147,108 +152,12 @@ void SObjectMixerEditorList::Construct(const FArguments& InArgs, TSharedRef<FObj
 SObjectMixerEditorList::~SObjectMixerEditorList()
 {	
 	HeaderRow.Reset();
-	
-	ListSearchBoxPtr.Reset();
-	ViewOptionsComboButton.Reset();
-	ListBoxContainerPtr.Reset();
 
 	FlushMemory(false);
 
-	ShowFilters.Reset();
 	TreeViewPtr.Reset();
 }
 
-TSharedRef<SWidget> SObjectMixerEditorList::GenerateToolbar()
-{
-	check(ListModelPtr.IsValid());
-	
-	TSharedRef<SHorizontalBox> ToolbarBox = SNew(SHorizontalBox);
-
-	// Add object button
-	ToolbarBox->AddSlot()
-   .HAlign(HAlign_Left)
-   .VAlign(VAlign_Center)
-   .AutoWidth()
-   .Padding(FMargin(0, 4))
-   [
-	   SNew(SPositiveActionButton)
-	   .Text(LOCTEXT("AddObject", "Add"))
-	   .OnGetMenuContent(FOnGetContent::CreateRaw(this, &SObjectMixerEditorList::OnGenerateAddObjectButtonMenu))
-   ];
-
-	ToolbarBox->AddSlot()
-	.HAlign(HAlign_Fill)
-	.VAlign(VAlign_Center)
-	.Padding(0.f, 1.f, 0.f, 1.f)
-	[
-		SAssignNew(ListSearchBoxPtr, SSearchBox)
-		.HintText(LOCTEXT("SearchHintText", "Search Scene Objects"))
-		.ToolTipText(LOCTEXT("ObjectMixerEditorList_TooltipText", "Search Scene Objects"))
-		.OnTextChanged_Raw(this, &SObjectMixerEditorList::OnListViewSearchTextChanged)
-	];
-
-	// Show Options
-	ToolbarBox->AddSlot()
-	.HAlign(HAlign_Right)
-	.VAlign(VAlign_Center)
-	.AutoWidth()
-	.Padding(8.f, 1.f, 0.f, 1.f)
-	[
-		SAssignNew( ViewOptionsComboButton, SComboButton )
-		.ContentPadding(4.f)
-		.ToolTipText(LOCTEXT("ShowOptions_Tooltip", "Show options to affect the visibility of items in the Object Mixer list"))
-		.ComboButtonStyle( FAppStyle::Get(), "SimpleComboButtonWithIcon" ) // Use the tool bar item style for this button
-		.OnGetMenuContent( this, &SObjectMixerEditorList::BuildShowOptionsMenu)
-		.HasDownArrow(false)
-		.ButtonContent()
-		[
-			SNew(SImage)
-			.ColorAndOpacity(FSlateColor::UseForeground())
-			.Image( FAppStyle::Get().GetBrush("Icons.Settings") )
-		]
-	];
-
-	return ToolbarBox;
-}
-
-TSharedRef<SWidget> SObjectMixerEditorList::OnGenerateAddObjectButtonMenu() const
-{
-	check(ListModelPtr.IsValid());
-
-	const TSet<TSubclassOf<AActor>> SubclassesOfActor = ListModelPtr.Pin()->GetObjectClassesToPlace();
-
-	if (SubclassesOfActor.Num() > 0)
-	{
-		TSet<UClass*> ClassesToPlace =
-		   ListModelPtr.Pin()->GetObjectFilter()->GetParentAndChildClassesFromSpecifiedClasses(
-			   SubclassesOfActor,
-			   ListModelPtr.Pin()->GetObjectFilter()->GetObjectMixerPlacementClassInclusionOptions());
-	
-		FMenuBuilder AddObjectButtonMenuBuilder = FMenuBuilder(true, nullptr);
-
-		for (const UClass* Class : ClassesToPlace)
-		{
-			if (const UActorFactory* Factory = GEditor->FindActorFactoryForActorClass(Class))
-			{
-				AddObjectButtonMenuBuilder.AddWidget(
-					SNew(SObjectMixerPlacementAssetMenuEntry, MakeShareable(new FPlaceableItem(*Factory->GetClass()))), FText::GetEmpty());
-			}
-		}
-
-		return AddObjectButtonMenuBuilder.MakeWidget();
-	}
-
-	return
-		SNew(SBox)
-		.Padding(5)
-		[
-			SNew(STextBlock)
-				.Text(LOCTEXT("NoPlaceableActorsDefinedWarning", "Please define some placeable actors in the\nfilter class by overriding GetObjectClassesToPlace."))
-				.Font(FAppStyle::Get().GetFontStyle("NormalFontItalic"))
-		]
-	;
-}
-
 void SObjectMixerEditorList::RebuildList()
 {
 	bShouldRebuild = false;
@@ -314,20 +223,24 @@ void SObjectMixerEditorList::SetTreeViewItems(const TArray<FObjectMixerEditorLis
 	TreeViewPtr->RequestListRefresh();
 }
 
-FString SObjectMixerEditorList::GetSearchStringFromSearchInputField() const
+TWeakPtr<FObjectMixerEditorListRow> SObjectMixerEditorList::GetSoloRow()
 {
-	return ensureAlwaysMsgf(ListSearchBoxPtr.IsValid(),
-		TEXT("%hs: ListSearchBoxPtr is not valid. Check to make sure it was created."), __FUNCTION__)
-	? ListSearchBoxPtr->GetText().ToString() : "";
+	return GetListModelPtr().Pin()->GetSoloRow();
 }
 
-void SObjectMixerEditorList::SetSearchStringInSearchInputField(const FString InSearchString) const
+void SObjectMixerEditorList::SetSoloRow(TSharedRef<FObjectMixerEditorListRow> InRow)
 {
-	if (ensureAlwaysMsgf(ListSearchBoxPtr.IsValid(),
-		TEXT("%hs: ListSearchBoxPtr is not valid. Check to make sure it was created."), __FUNCTION__))
-	{
-		ListSearchBoxPtr->SetText(FText::FromString(InSearchString));
-	}
+	GetListModelPtr().Pin()->SetSoloRow(InRow);
+}
+
+void SObjectMixerEditorList::ClearSoloRow()
+{
+	GetListModelPtr().Pin()->ClearSoloRow();
+}
+
+FString SObjectMixerEditorList::GetSearchStringFromSearchInputField() const
+{
+	return ListModelPtr.Pin()->GetMainPanelModel().Pin()->GetSearchStringFromSearchInputField();
 }
 
 void SObjectMixerEditorList::ExecuteListViewSearchOnAllRows(
@@ -398,32 +311,22 @@ EObjectMixerTreeViewMode SObjectMixerEditorList::GetTreeViewMode()
 {
 	const TSharedPtr<FObjectMixerEditorList> PinnedListModel = GetListModelPtr().Pin();
 	check(PinnedListModel);
-	
-	return PinnedListModel->GetTreeViewMode();
-}
 
-void SObjectMixerEditorList::SetTreeViewMode(EObjectMixerTreeViewMode InViewMode)
-{
-	if (TSharedPtr<FObjectMixerEditorList> PinnedListModel = GetListModelPtr().Pin())
-	{
-		PinnedListModel->SetTreeViewMode(InViewMode);
-	}
+	const TSharedPtr<FObjectMixerEditorMainPanel> PinnedMainPanel = PinnedListModel->GetMainPanelModel().Pin();
+	check(PinnedMainPanel);
+	
+	return PinnedMainPanel->GetTreeViewMode();
 }
 
-void SObjectMixerEditorList::ToggleFilterActive(const FString& FilterName)
+const TArray<TSharedRef<IObjectMixerEditorListFilter>>& SObjectMixerEditorList::GetShowFilters()
 {
-	if (const TSharedRef<IObjectMixerEditorListFilter>* Match =
-		Algo::FindByPredicate(ShowFilters,
-		[&FilterName](TSharedRef<IObjectMixerEditorListFilter> Comparator)
-		{
-			return Comparator->GetFilterName().Equals(FilterName);
-		}))
-	{
-		const TSharedRef<IObjectMixerEditorListFilter> Filter = *Match;
-		Filter->ToggleFilterActive();
+	const TSharedPtr<FObjectMixerEditorList> PinnedListModel = GetListModelPtr().Pin();
+	check(PinnedListModel);
 
-		EvaluateIfRowsPassFilters();
-	}
+	const TSharedPtr<FObjectMixerEditorMainPanel> PinnedMainPanel = PinnedListModel->GetMainPanelModel().Pin();
+	check(PinnedMainPanel);
+	
+	return PinnedMainPanel->GetShowFilters();
 }
 
 void SObjectMixerEditorList::EvaluateIfRowsPassFilters(const bool bShouldRefreshAfterward)
@@ -433,7 +336,7 @@ void SObjectMixerEditorList::EvaluateIfRowsPassFilters(const bool bShouldRefresh
 	TSet<TSharedRef<IObjectMixerEditorListFilter>> MatchAnyOfFilters;
 	TSet<TSharedRef<IObjectMixerEditorListFilter>> MatchAllOfFilters;
 
-	for (const TSharedRef<IObjectMixerEditorListFilter>& Filter : ShowFilters)
+	for (const TSharedRef<IObjectMixerEditorListFilter>& Filter : GetShowFilters())
 	{
 		if (Filter->GetFilterMatchType() ==
 			IObjectMixerEditorListFilter::EObjectMixerEditorListFilterMatchType::MatchAll)
@@ -445,22 +348,33 @@ void SObjectMixerEditorList::EvaluateIfRowsPassFilters(const bool bShouldRefresh
 			MatchAnyOfFilters.Add(Filter);
 		}
 	}
-	
-	for (const FObjectMixerEditorListRowPtr& Row : TreeViewRootObjects)
+
+	const TFunction<void(const TArray<FObjectMixerEditorListRowPtr>&)> RecursiveRowIterator =
+		[&MatchAnyOfFilters, &MatchAllOfFilters, &RecursiveRowIterator](const TArray<FObjectMixerEditorListRowPtr>& InObjects)
 	{
-		if (Row.IsValid() && Row->GetRowType() == FObjectMixerEditorListRow::MatchingObject)
+		for (const FObjectMixerEditorListRowPtr& Row : InObjects)
 		{
-			auto Projection = [&Row](const TSharedRef<IObjectMixerEditorListFilter>& Filter)
+			if (Row.IsValid())
 			{
-				return Filter->GetIsFilterActive() ? Filter->DoesItemPassFilter(Row) : true;
-			};
-			
-			const bool bPassesAnyOf = MatchAnyOfFilters.Num() ? Algo::AnyOf(MatchAnyOfFilters, Projection) : true;
-			const bool bPassesAllOf = MatchAllOfFilters.Num() ? Algo::AllOf(MatchAllOfFilters, Projection) : true;
+				auto FilterProjection = [&Row](const TSharedRef<IObjectMixerEditorListFilter>& Filter)
+				{
+					return Filter->GetIsFilterActive() ? Filter->DoesItemPassFilter(Row) : true;
+				};
 			
-			Row->SetDoesRowPassFilters(bPassesAnyOf && bPassesAllOf);
+				const bool bPassesAnyOf = MatchAnyOfFilters.Num() ? Algo::AnyOf(MatchAnyOfFilters, FilterProjection) : true;
+				const bool bPassesAllOf = MatchAllOfFilters.Num() ? Algo::AllOf(MatchAllOfFilters, FilterProjection) : true;
+		
+				Row->SetDoesRowPassFilters(bPassesAnyOf && bPassesAllOf);
+				
+				if (Row->GetChildCount() > 0)
+				{
+					RecursiveRowIterator(Row->GetChildRows());
+				}
+			}
 		}
-	}
+	};
+
+	RecursiveRowIterator(TreeViewRootObjects);
 
 	if (bShouldRefreshAfterward)
 	{
@@ -573,7 +487,7 @@ TSharedRef<SWidget> SObjectMixerEditorList::GenerateHeaderRowContextMenu() const
 				})
 			),
 			Hook,
-			EUserInterfaceActionType::Check
+			EUserInterfaceActionType::ToggleButton
 		);
 	}
 
@@ -685,6 +599,7 @@ void SObjectMixerEditorList::AddBuiltinColumnsToHeaderRow()
 TSharedPtr<SHeaderRow> SObjectMixerEditorList::GenerateHeaderRow()
 {
 	check(ListModelPtr.IsValid());
+	check(ListModelPtr.Pin()->GetMainPanelModel().Pin());
 	check(HeaderRow);
 	
 	TMap<FName, bool> LastVisibleColumns;
@@ -698,7 +613,7 @@ TSharedPtr<SHeaderRow> SObjectMixerEditorList::GenerateHeaderRow()
 	ListViewColumns.Empty(ListViewColumns.Num());
 	
 	TSet<UClass*> SpecifiedClasses =
-		ListModelPtr.Pin()->GetObjectFilter()->GetParentAndChildClassesFromSpecifiedClasses(
+		ListModelPtr.Pin()->GetMainPanelModel().Pin()->GetObjectFilter()->GetParentAndChildClassesFromSpecifiedClasses(
 			ObjectClassesToFilterCache, PropertyInheritanceInclusionOptionsCache);
 	
 	for (const UClass* Class : SpecifiedClasses)
@@ -827,216 +742,6 @@ TSharedPtr<SHeaderRow> SObjectMixerEditorList::GenerateHeaderRow()
 	return HeaderRow;
 }
 
-void SObjectMixerEditorList::SetupFilters()
-{
-	
-}
-
-TSharedRef<SWidget> SObjectMixerEditorList::OnGenerateFilterClassMenu()
-{
-	FMenuBuilder MenuBuilder(true, NULL);
-
-	TArray<UClass*> DerivedClasses;
-	GetDerivedClasses(UObjectMixerObjectFilter::StaticClass(), DerivedClasses, true);
-
-	DerivedClasses.Remove(UObjectMixerObjectFilter::StaticClass());
-	DerivedClasses.Remove(UObjectMixerBlueprintObjectFilter::StaticClass());
-
-	DerivedClasses.Sort([](UClass& A, UClass& B)
-	{
-		return A.GetFName().LexicalLess(B.GetFName());
-	});
-
-	if (DerivedClasses.Num())
-	{
-		const TSharedPtr<FObjectMixerEditorList, ESPMode::ThreadSafe> PinnedList = ListModelPtr.Pin();
-		check(PinnedList);
-		
-		MenuBuilder.BeginSection(NAME_None, LOCTEXT("SelectClassMenuSection", "Select Class"));
-		{
-			for (UClass* DerivedClass : DerivedClasses)
-			{
-				if (IsValid(DerivedClass))
-				{
-					if (DerivedClass->GetName().StartsWith(TEXT("SKEL_")) || DerivedClass->GetName().StartsWith(TEXT("REINST_")))
-					{
-						continue;
-					}
-
-					if (DerivedClass->HasAnyClassFlags(CLASS_Abstract | CLASS_HideDropDown | CLASS_Deprecated))
-					{
-						continue;
-					}
-					
-					MenuBuilder.AddMenuEntry(
-					   FText::FromName(DerivedClass->GetFName()),
-					   FText::GetEmpty(),
-					   FSlateIcon(),
-					   FUIAction(
-						   FExecuteAction::CreateSP(PinnedList.ToSharedRef(), &FObjectMixerEditorList::SetObjectFilterClass, DerivedClass),
-						   FCanExecuteAction::CreateLambda([](){ return true; }),
-						   FIsActionChecked::CreateSP(PinnedList.ToSharedRef(), &FObjectMixerEditorList::IsClassSelected, DerivedClass)
-					   ),
-					   NAME_None,
-					   EUserInterfaceActionType::RadioButton
-				   );
-				}
-			}
-		}
-		MenuBuilder.EndSection();
-	}
-	else
-	{
-		MenuBuilder.AddMenuEntry(LOCTEXT("NoFilterClassesAvailable", "No filter classes available."), FText::GetEmpty(), FSlateIcon(), FUIAction());
-	}
-
-	TSharedRef<SWidget> Widget = MenuBuilder.MakeWidget();
-	FChildren* ChildWidgets = Widget->GetChildren();
-	for (int32 ChildItr = 0; ChildItr < ChildWidgets->Num(); ChildItr++)
-	{
-		const TSharedRef<SWidget>& Child = ChildWidgets->GetChildAt(ChildItr);
-
-		Child->EnableToolTipForceField(false);
-	}
-	Widget->EnableToolTipForceField(false);
-	
-	return Widget;
-}
-
-TSharedRef<SWidget> SObjectMixerEditorList::BuildShowOptionsMenu()
-{
-	FMenuBuilder ShowOptionsMenuBuilder = FMenuBuilder(true, nullptr);
-
-	ShowOptionsMenuBuilder.BeginSection("ListViewOptions", LOCTEXT("FilterClassManagementSection", "Filter Class Management"));
-	{
-		// Filter Class Management Button
-		const TSharedRef<SWidget> FilterClassManagementButton =
-			SNew(SBox)
-			.Padding(8, 0)
-			[
-				SNew(SComboButton)
-				.ToolTipText(LOCTEXT("FilterClassManagementButton_Tooltip", "Select a filter class"))
-				.ContentPadding(FMargin(4, 0.5f))
-				.ComboButtonStyle(&FAppStyle::Get().GetWidgetStyle<FComboButtonStyle>("ComboButton"))
-				.OnGetMenuContent(this, &SObjectMixerEditorList::OnGenerateFilterClassMenu)
-				.ForegroundColor(FStyleColors::Foreground)
-				.MenuPlacement(EMenuPlacement::MenuPlacement_MenuRight)
-				.ButtonContent()
-				[
-					SNew(SHorizontalBox)
-
-					+ SHorizontalBox::Slot()
-					.Padding(0, 1, 4, 0)
-					.AutoWidth()
-					[
-						SNew(SImage)
-						.Image(FAppStyle::Get().GetBrush("Icons.Filter"))
-						.ColorAndOpacity(FSlateColor::UseForeground())
-					]
-
-					+ SHorizontalBox::Slot()
-					.Padding(0, 1, 0, 0)
-					.AutoWidth()
-					[
-						SNew(STextBlock)
-						.Text(LOCTEXT("FilterClassToolbarButton", "Object Filter Class"))
-					]
-				]
-			];
-
-		ShowOptionsMenuBuilder.AddWidget(FilterClassManagementButton, FText::GetEmpty());
-	}
-	ShowOptionsMenuBuilder.EndSection();
-
-	// Add List View Mode Options
-	ShowOptionsMenuBuilder.BeginSection("ListViewOptions", LOCTEXT("ListViewOptionsSection", "List View Options"));
-	{
-		// Foreach on uenum
-		const FString EnumPath = "/Script/ObjectMixerEditor.EObjectMixerTreeViewMode";
-		if (const UEnum* EnumPtr = FindObject<UEnum>(nullptr, *EnumPath, true))
-		{
-			for (int32 EnumItr = 0; EnumItr < EnumPtr->GetMaxEnumValue(); EnumItr++)
-			{
-				EObjectMixerTreeViewMode EnumValue = (EObjectMixerTreeViewMode)EnumItr;
-				
-				ShowOptionsMenuBuilder.AddMenuEntry(
-					EnumPtr->GetDisplayNameTextByIndex(EnumItr),
-					EnumPtr->GetToolTipTextByIndex(EnumItr),
-					FSlateIcon(),
-					FUIAction(
-						FExecuteAction::CreateLambda([this, EnumValue]()
-						{
-							SetTreeViewMode(EnumValue);
-						}),
-						FCanExecuteAction::CreateLambda([](){ return true; }),
-						FIsActionChecked::CreateLambda([this, EnumValue]()
-						{
-							return GetTreeViewMode() == EnumValue;
-						})
-					),
-					NAME_None,
-					EUserInterfaceActionType::RadioButton
-				);
-			}
-		}
-	}
-	ShowOptionsMenuBuilder.EndSection();
-
-	ShowOptionsMenuBuilder.BeginSection("MiscOptionsSection", LOCTEXT("MiscOptionsSection","Misc"));
-	{
-		ShowOptionsMenuBuilder.AddMenuEntry(
-		   FText::FromString("Open Generic Object Mixer Instance"),
-		   FText::FromString("Open Generic Object Mixer Instance"),
-		   FSlateIcon(),
-		   FUIAction(FExecuteAction::CreateLambda([]()
-		   {
-			   FGlobalTabmanager::Get()->TryInvokeTab(FObjectMixerEditorModule::Get().GetTabSpawnerId());
-		   })));
-
-		ShowOptionsMenuBuilder.AddMenuEntry(FText::FromString("Refresh List"), FText::FromString("Refresh"), FSlateIcon(),
-			FUIAction(FExecuteAction::CreateRaw(this, &SObjectMixerEditorList::GenerateTreeView)));
-	}
-	ShowOptionsMenuBuilder.EndSection();
-
-	if (ShowFilters.Num())
-	{
-		ShowOptionsMenuBuilder.BeginSection("", LOCTEXT("ShowOptions_ShowSectionHeading", "Show"));
-		{
-			// Add show filters
-			auto AddFiltersLambda = [this, &ShowOptionsMenuBuilder](const TSharedRef<IObjectMixerEditorListFilter>& InFilter)
-			{
-				const FString& FilterName = InFilter->GetFilterName();
-			
-				ShowOptionsMenuBuilder.AddMenuEntry(
-				   InFilter->GetFilterButtonLabel(),
-				   InFilter->GetFilterButtonToolTip(),
-				   FSlateIcon(),
-				   FUIAction(
-					   FExecuteAction::CreateLambda(
-						   [this, FilterName]()
-						   {
-							   ToggleFilterActive(FilterName);
-						   }
-						),
-					   FCanExecuteAction(),
-					   FIsActionChecked::CreateSP( InFilter, &IObjectMixerEditorListFilter::GetIsFilterActive )
-				   ),
-				   NAME_None,
-				   EUserInterfaceActionType::ToggleButton
-			   );
-			};
-
-			for (const TSharedRef<IObjectMixerEditorListFilter>& Filter : ShowFilters)
-			{
-				AddFiltersLambda(Filter);
-			}
-		}
-		ShowOptionsMenuBuilder.EndSection();
-	}
-
-	return ShowOptionsMenuBuilder.MakeWidget();
-}
-
 void SObjectMixerEditorList::FlushMemory(const bool bShouldKeepMemoryAllocated)
 {
 	if (bShouldKeepMemoryAllocated)
@@ -1068,11 +773,6 @@ void SObjectMixerEditorList::SetAllGroupsCollapsed()
 	}
 }
 
-void SObjectMixerEditorList::OnListViewSearchTextChanged(const FText& Text)
-{
-	ExecuteListViewSearchOnAllRows(Text.ToString(), true);
-}
-
 void SObjectMixerEditorList::CacheTreeState()
 {
 	struct Local
@@ -1148,12 +848,15 @@ void SObjectMixerEditorList::RestoreTreeState(const bool bFlushCache)
 
 void SObjectMixerEditorList::BuildPerformanceCacheAndGenerateHeaderIfNeeded()
 {
+	check(ListModelPtr.IsValid());
+	check(ListModelPtr.Pin()->GetMainPanelModel().Pin());
+	
 	// If any of the following overrides change, we need to regenerate the header row. Otherwise skip regeneration for performance reasons.
 	// GetObjectClassesToFilter, GetColumnsToShowByDefault, GetColumnsToExclude,
 	// GetForceAddedColumns, GetObjectMixerPropertyInheritanceInclusionOptions, ShouldIncludeUnsupportedProperties
 	bool bNeedToGenerateHeaders = false;
 	
-	const TObjectPtr<UObjectMixerObjectFilter> SelectedFilter = ListModelPtr.Pin()->GetObjectFilter();
+	const TObjectPtr<UObjectMixerObjectFilter> SelectedFilter = ListModelPtr.Pin()->GetMainPanelModel().Pin()->GetObjectFilter();
 	if (!SelectedFilter)
 	{
 		UE_LOG(LogObjectMixerEditor, Display, TEXT("%hs: No classes defined in UObjectMixerObjectFilter class."), __FUNCTION__);
@@ -1284,7 +987,7 @@ void SObjectMixerEditorList::GenerateTreeView()
 			CreatedObjectMap.Add(Object, TopLevelRow);
 
 			// If the view is not in flat mode, we need to consider the hierarchy of outliner folders/attach parents as desired
-			if (ListModelPtr.Pin()->GetTreeViewMode() != EObjectMixerTreeViewMode::Flat)
+			if (GetTreeViewMode() != EObjectMixerTreeViewMode::Flat)
 			{
 				AActor* BaseActor = Cast<AActor>(Object);
 				
@@ -1295,7 +998,7 @@ void SObjectMixerEditorList::GenerateTreeView()
 						BaseActor = Object->GetTypedOuter<AActor>();
 
 						// If it's not flat or folder view mode, we need to find or create the container object for the actor that owns the matching component
-						if (ListModelPtr.Pin()->GetTreeViewMode() != EObjectMixerTreeViewMode::Folder)
+						if (GetTreeViewMode() != EObjectMixerTreeViewMode::Folder)
 						{
 							TSharedPtr<FObjectMixerEditorListRow> OwningActorRow;
 
@@ -1329,7 +1032,7 @@ void SObjectMixerEditorList::GenerateTreeView()
 					while (AActor* AttachParent = BaseActor->GetAttachParentActor())
 					{
 						// Make a row for each attach parent up the chain until we reach the top if not in flat/folder mode
-						if (ListModelPtr.Pin()->GetTreeViewMode() != EObjectMixerTreeViewMode::Folder)
+						if (GetTreeViewMode() != EObjectMixerTreeViewMode::Folder)
 						{
 							TSharedPtr<FObjectMixerEditorListRow> OwningActorRow;
 
@@ -1357,8 +1060,8 @@ void SObjectMixerEditorList::GenerateTreeView()
 					}
 
 					// Now consider folder hierarchy for the base actor if desired
-					if (ListModelPtr.Pin()->GetTreeViewMode() == EObjectMixerTreeViewMode::Folder ||
-						ListModelPtr.Pin()->GetTreeViewMode() == EObjectMixerTreeViewMode::FolderObjectSubObject)
+					if (GetTreeViewMode() == EObjectMixerTreeViewMode::Folder ||
+						GetTreeViewMode() == EObjectMixerTreeViewMode::FolderObjectSubObject)
 					{
 						if (BaseActor)
 						{
@@ -1402,6 +1105,11 @@ void SObjectMixerEditorList::GenerateTreeView()
 
 	TreeViewRootObjects.StableSort(SortByTypeThenName);
 
+	if (TSharedPtr<FObjectMixerEditorMainPanel> MainPanel = ListModelPtr.Pin()->GetMainPanelModel().Pin())
+	{
+		MainPanel->RebuildCategorySelector();
+	}
+
 	RefreshList();
 
 	RestoreTreeState();
diff --git a/Engine/Source/Editor/ObjectMixer/ObjectMixer/Private/Views/MainPanel/ObjectMixerEditorMainPanel.cpp b/Engine/Source/Editor/ObjectMixer/ObjectMixer/Private/Views/MainPanel/ObjectMixerEditorMainPanel.cpp
--- a/Engine/Source/Editor/ObjectMixer/ObjectMixer/Private/Views/MainPanel/ObjectMixerEditorMainPanel.cpp
+++ b/Engine/Source/Editor/ObjectMixer/ObjectMixer/Private/Views/MainPanel/ObjectMixerEditorMainPanel.cpp
@@ -2,14 +2,13 @@
 
 #include "Views/MainPanel/ObjectMixerEditorMainPanel.h"
 
-#include "ObjectMixerEditorModule.h"
 #include "ObjectMixerEditorProjectSettings.h"
 #include "Views/List/ObjectMixerEditorList.h"
 #include "Views/MainPanel/SObjectMixerEditorMainPanel.h"
 
-FObjectMixerEditorMainPanel::FObjectMixerEditorMainPanel()
+void FObjectMixerEditorMainPanel::Init()
 {
-	EditorList = MakeShared<FObjectMixerEditorList>();
+	RegenerateListModel();
 }
 
 TSharedRef<SWidget> FObjectMixerEditorMainPanel::GetOrCreateWidget()
@@ -22,41 +21,161 @@ TSharedRef<SWidget> FObjectMixerEditorMainPanel::GetOrCreateWidget()
 	return MainPanelWidget.ToSharedRef();
 }
 
+void FObjectMixerEditorMainPanel::RegenerateListModel()
+{
+	EditorListModel.Reset();
+	
+	EditorListModel = MakeShared<FObjectMixerEditorList>(SharedThis(this));
+}
+
 void FObjectMixerEditorMainPanel::RequestRebuildList() const
 {
-	if (EditorList.IsValid())
+	if (EditorListModel.IsValid())
 	{
-		EditorList->RequestRebuildList();
+		EditorListModel->RequestRebuildList();
 	}
 }
 
 void FObjectMixerEditorMainPanel::RefreshList() const
 {
-	if (EditorList.IsValid())
+	if (EditorListModel.IsValid())
 	{
-		EditorList->RefreshList();
+		EditorListModel->RefreshList();
 	}
 }
 
-void FObjectMixerEditorMainPanel::OnClassSelectionChanged(UClass* InNewClass) const
+void FObjectMixerEditorMainPanel::RebuildCategorySelector()
+{
+	MainPanelWidget->RebuildCategorySelector();
+}
+
+FString FObjectMixerEditorMainPanel::GetSearchStringFromSearchInputField() const
+{
+	return MainPanelWidget->GetSearchStringFromSearchInputField();
+}
+
+void FObjectMixerEditorMainPanel::OnClassSelectionChanged(UClass* InNewClass)
+{
+	SetObjectFilterClass(InNewClass);
+}
+
+TObjectPtr<UClass> FObjectMixerEditorMainPanel::GetClassSelection() const
+{
+	return GetObjectFilterClass();
+}
+
+bool FObjectMixerEditorMainPanel::IsClassSelected(UClass* InNewClass) const
 {
-	if (TSharedPtr<FObjectMixerEditorList, ESPMode::ThreadSafe> PinnedList = GetEditorList().Pin())
+	return InNewClass == GetClassSelection();
+}
+
+UObjectMixerObjectFilter* FObjectMixerEditorMainPanel::GetObjectFilter()
+{
+	if (!ObjectFilterPtr.IsValid())
 	{
-		PinnedList->SetObjectFilterClass(InNewClass);
+		CacheObjectFilterObject();
 	}
+
+	return ObjectFilterPtr.Get();
 }
 
-TObjectPtr<UClass> FObjectMixerEditorMainPanel::GetClassSelection() const
+void FObjectMixerEditorMainPanel::CacheObjectFilterObject()
 {
-	if (TSharedPtr<FObjectMixerEditorList, ESPMode::ThreadSafe> PinnedList = GetEditorList().Pin())
+	if (ObjectFilterPtr.IsValid())
+	{
+		ObjectFilterPtr.Reset();
+	}
+	
+	if (const UClass* Class = GetObjectFilterClass())
 	{
-		return PinnedList->GetObjectFilterClass();
+		ObjectFilterPtr = TStrongObjectPtr(NewObject<UObjectMixerObjectFilter>(GetTransientPackage(), Class));
 	}
+}
 
-	return nullptr;
+const TArray<TSharedRef<IObjectMixerEditorListFilter>>& FObjectMixerEditorMainPanel::GetShowFilters() const
+{
+	return MainPanelWidget->GetShowFilters();	
 }
 
-bool FObjectMixerEditorMainPanel::IsClassSelected(UClass* InNewClass) const
+void FObjectMixerEditorMainPanel::AddObjectsToCategory(const FName& CategoryName, const TSet<FSoftObjectPath>& ObjectsToAdd) const
 {
-	return InNewClass == GetClassSelection();
+	if (const TSubclassOf<UObjectMixerObjectFilter> Filter = GetObjectFilterClass())
+	{
+		if (UObjectMixerEditorProjectSettings* Settings = GetMutableDefault<UObjectMixerEditorProjectSettings>())
+		{
+			Settings->AddObjectsToCategory(Filter->GetFName(), CategoryName, ObjectsToAdd);
+
+			OnObjectMixerCategoryMapChanged.Broadcast();
+		}
+	}
+}
+
+void FObjectMixerEditorMainPanel::RemoveObjectsFromCategory(const FName& CategoryName, const TSet<FSoftObjectPath>& ObjectsToRemove) const
+{
+	if (const TSubclassOf<UObjectMixerObjectFilter> Filter = GetObjectFilterClass())
+	{
+		if (UObjectMixerEditorProjectSettings* Settings = GetMutableDefault<UObjectMixerEditorProjectSettings>())
+		{
+			Settings->RemoveObjectsFromCategory(Filter->GetFName(), CategoryName, ObjectsToRemove);
+
+			OnObjectMixerCategoryMapChanged.Broadcast();
+		}
+	}
+}
+
+void FObjectMixerEditorMainPanel::RemoveCategory(const FName& CategoryName) const
+{
+	if (const TSubclassOf<UObjectMixerObjectFilter> Filter = GetObjectFilterClass())
+	{
+		if (UObjectMixerEditorProjectSettings* Settings = GetMutableDefault<UObjectMixerEditorProjectSettings>())
+		{
+			Settings->RemoveCategory(Filter->GetFName(), CategoryName);
+
+			OnObjectMixerCategoryMapChanged.Broadcast();
+		}
+	}
+}
+
+bool FObjectMixerEditorMainPanel::IsObjectInCategory(const FName& CategoryName, const FSoftObjectPath& InObject) const
+{
+	if (const TSubclassOf<UObjectMixerObjectFilter> Filter = GetObjectFilterClass())
+	{
+		if (UObjectMixerEditorProjectSettings* Settings = GetMutableDefault<UObjectMixerEditorProjectSettings>())
+		{
+			return Settings->IsObjectInCategory(Filter->GetFName(), CategoryName, InObject);
+		}
+	}
+
+	return false;
+}
+
+TSet<FName> FObjectMixerEditorMainPanel::GetCategoriesForObject(const FSoftObjectPath& InObject) const
+{
+	if (const TSubclassOf<UObjectMixerObjectFilter> Filter = GetObjectFilterClass())
+	{
+		if (UObjectMixerEditorProjectSettings* Settings = GetMutableDefault<UObjectMixerEditorProjectSettings>())
+		{
+			return Settings->GetCategoriesForObject(Filter->GetFName(), InObject);
+		}
+	}
+
+	return {};
+}
+
+TSet<FName> FObjectMixerEditorMainPanel::GetAllCategories() const
+{
+	if (const TSubclassOf<UObjectMixerObjectFilter> Filter = GetObjectFilterClass())
+	{
+		if (UObjectMixerEditorProjectSettings* Settings = GetMutableDefault<UObjectMixerEditorProjectSettings>())
+		{
+			return Settings->GetAllCategories(Filter->GetFName());
+		}
+	}
+
+	return {};
+}
+
+const TSet<FName>& FObjectMixerEditorMainPanel::GetCurrentCategorySelection() const
+{
+	return MainPanelWidget->GetCurrentCategorySelection();
 }
diff --git a/Engine/Source/Editor/ObjectMixer/ObjectMixer/Private/Views/MainPanel/SObjectMixerEditorMainPanel.cpp b/Engine/Source/Editor/ObjectMixer/ObjectMixer/Private/Views/MainPanel/SObjectMixerEditorMainPanel.cpp
--- a/Engine/Source/Editor/ObjectMixer/ObjectMixer/Private/Views/MainPanel/SObjectMixerEditorMainPanel.cpp
+++ b/Engine/Source/Editor/ObjectMixer/ObjectMixer/Private/Views/MainPanel/SObjectMixerEditorMainPanel.cpp
@@ -3,38 +3,557 @@
 #include "Views/MainPanel/SObjectMixerEditorMainPanel.h"
 
 #include "ObjectMixerEditorLog.h"
+#include "ObjectMixerEditorModule.h"
 #include "Views/List/ObjectMixerEditorList.h"
 #include "Views/MainPanel/ObjectMixerEditorMainPanel.h"
+#include "Views/Widgets/SObjectMixerPlacementAssetMenuEntry.h"
 
+#include "Framework/MultiBox/MultiBoxBuilder.h"
 #include "Kismet2/SClassPickerDialog.h"
+#include "PlacementMode/Public/IPlacementModeModule.h"
+#include "SPositiveActionButton.h"
+#include "Styling/StyleColors.h"
+#include "Views/List/ObjectMixerEditorListFilters/IObjectMixerEditorListFilter.h"
+#include "Views/List/ObjectMixerEditorListFilters/ObjectMixerEditorListFilter_Category.h"
+#include "Widgets/Input/SCheckBox.h"
+#include "Widgets/Input/SComboButton.h"
 #include "Widgets/Input/SMultiLineEditableTextBox.h"
-#include "Widgets/Layout/SSplitter.h"
+#include "Widgets/Input/SSearchBox.h"
+#include "Widgets/Layout/SWidgetSwitcher.h"
+#include "Widgets/Layout/SWrapBox.h"
+#include "Widgets/Text/STextBlock.h"
 
 #define LOCTEXT_NAMESPACE "ObjectMixerEditor"
 
 void SObjectMixerEditorMainPanel::Construct(
 	const FArguments& InArgs, const TSharedRef<FObjectMixerEditorMainPanel>& InMainPanel)
 {
-	check(InMainPanel->GetEditorList().IsValid());
+	check(InMainPanel->GetEditorListModel().IsValid());
 
-	MainPanel = InMainPanel;
+	MainPanelModel = InMainPanel;
 	
 	ChildSlot
 	[
-		SNew(SSplitter)
-		.Orientation(Orient_Vertical)
-		.ResizeMode(ESplitterResizeMode::FixedSize)
+		SNew(SVerticalBox)
 
-		+SSplitter::Slot()
+		+SVerticalBox::Slot()
+		.HAlign(HAlign_Fill)
+		.AutoHeight()
 		[
-			MainPanel.Pin()->GetEditorList().Pin()->GetOrCreateWidget()
+			GenerateToolbar()
 		]
+
+		+SVerticalBox::Slot()
+		.Padding(8, 2, 8, 7)
+		.AutoHeight()
+		[
+			SAssignNew(CategorySelectorBox, SWrapBox)
+			.UseAllottedSize(true)
+			.InnerSlotPadding(FVector2D(4,4))
+		]
+
+		+SVerticalBox::Slot()
+		.HAlign(HAlign_Fill)
+		.VAlign(VAlign_Fill)
+		[
+			MainPanelModel.Pin()->GetEditorListModel().Pin()->GetOrCreateWidget()
+		]
+	];
+
+	InMainPanel->GetOnObjectMixerCategoryMapChanged().AddRaw(this, &SObjectMixerEditorMainPanel::RebuildCategorySelector);
+
+	ShowFilters.Add(MakeShared<FObjectMixerEditorListFilter_Category>());
+}
+
+TSharedRef<SWidget> SObjectMixerEditorMainPanel::GenerateToolbar()
+{
+	TSharedRef<SHorizontalBox> ToolbarBox = SNew(SHorizontalBox);
+
+	// Add object button
+	ToolbarBox->AddSlot()
+   .HAlign(HAlign_Left)
+   .VAlign(VAlign_Center)
+   .AutoWidth()
+   .Padding(FMargin(0, 4))
+   [
+	   SNew(SPositiveActionButton)
+	   .Text(LOCTEXT("AddObject", "Add"))
+	   .OnGetMenuContent(FOnGetContent::CreateRaw(this, &SObjectMixerEditorMainPanel::OnGenerateAddObjectButtonMenu))
+   ];
+
+	ToolbarBox->AddSlot()
+	.HAlign(HAlign_Fill)
+	.VAlign(VAlign_Center)
+	.Padding(0.f, 1.f, 0.f, 1.f)
+	[
+		SAssignNew(SearchBoxPtr, SSearchBox)
+		.HintText(LOCTEXT("SearchHintText", "Search Scene Objects"))
+		.ToolTipText(LOCTEXT("ObjectMixerEditorList_TooltipText", "Search Scene Objects"))
+		.OnTextChanged_Raw(this, &SObjectMixerEditorMainPanel::OnSearchTextChanged)
+	];
+
+	// Show Options
+	ToolbarBox->AddSlot()
+	.HAlign(HAlign_Right)
+	.VAlign(VAlign_Center)
+	.AutoWidth()
+	.Padding(8.f, 1.f, 0.f, 1.f)
+	[
+		SAssignNew( ViewOptionsComboButton, SComboButton )
+		.ContentPadding(4.f)
+		.ToolTipText(LOCTEXT("ShowOptions_Tooltip", "Show options to affect the visibility of items in the Object Mixer list"))
+		.ComboButtonStyle( FAppStyle::Get(), "SimpleComboButtonWithIcon" ) // Use the tool bar item style for this button
+		.OnGetMenuContent( this, &SObjectMixerEditorMainPanel::BuildShowOptionsMenu)
+		.HasDownArrow(false)
+		.ButtonContent()
+		[
+			SNew(SImage)
+			.ColorAndOpacity(FSlateColor::UseForeground())
+			.Image( FAppStyle::Get().GetBrush("Icons.Settings") )
+		]
+	];
+
+	return ToolbarBox;
+}
+
+TSharedRef<SWidget> SObjectMixerEditorMainPanel::OnGenerateAddObjectButtonMenu() const
+{
+	const TSet<TSubclassOf<AActor>> SubclassesOfActor = MainPanelModel.Pin()->GetObjectClassesToPlace();
+
+	if (SubclassesOfActor.Num() > 0)
+	{
+		TSet<UClass*> ClassesToPlace =
+		   MainPanelModel.Pin()->GetObjectFilter()->GetParentAndChildClassesFromSpecifiedClasses(
+			   SubclassesOfActor,
+			   MainPanelModel.Pin()->GetObjectFilter()->GetObjectMixerPlacementClassInclusionOptions());
+	
+		FMenuBuilder AddObjectButtonMenuBuilder = FMenuBuilder(true, nullptr);
+
+		for (const UClass* Class : ClassesToPlace)
+		{
+			if (const UActorFactory* Factory = GEditor->FindActorFactoryForActorClass(Class))
+			{
+				AddObjectButtonMenuBuilder.AddWidget(
+					SNew(SObjectMixerPlacementAssetMenuEntry, MakeShareable(new FPlaceableItem(*Factory->GetClass()))), FText::GetEmpty());
+			}
+		}
+
+		return AddObjectButtonMenuBuilder.MakeWidget();
+	}
+
+	return
+		SNew(SBox)
+		.Padding(5)
+		[
+			SNew(STextBlock)
+				.Text(LOCTEXT("NoPlaceableActorsDefinedWarning", "Please define some placeable actors in the\nfilter class by overriding GetObjectClassesToPlace."))
+				.Font(FAppStyle::Get().GetFontStyle("NormalFontItalic"))
+		]
+	;
+}
+
+TSharedRef<SWidget> SObjectMixerEditorMainPanel::OnGenerateFilterClassMenu()
+{
+	FMenuBuilder MenuBuilder(true, nullptr);
+
+	TArray<UClass*> DerivedClasses;
+	GetDerivedClasses(UObjectMixerObjectFilter::StaticClass(), DerivedClasses, true);
+
+	DerivedClasses.Remove(UObjectMixerObjectFilter::StaticClass());
+	DerivedClasses.Remove(UObjectMixerBlueprintObjectFilter::StaticClass());
+
+	DerivedClasses.Sort([](UClass& A, UClass& B)
+	{
+		return A.GetFName().LexicalLess(B.GetFName());
+	});
+
+	if (DerivedClasses.Num())
+	{
+		check(MainPanelModel.Pin());
+		
+		MenuBuilder.BeginSection(NAME_None, LOCTEXT("SelectClassMenuSection", "Select Class"));
+		{
+			for (UClass* DerivedClass : DerivedClasses)
+			{
+				if (IsValid(DerivedClass))
+				{
+					if (DerivedClass->GetName().StartsWith(TEXT("SKEL_")) || DerivedClass->GetName().StartsWith(TEXT("REINST_")))
+					{
+						continue;
+					}
+
+					if (DerivedClass->HasAnyClassFlags(CLASS_Abstract | CLASS_HideDropDown | CLASS_Deprecated))
+					{
+						continue;
+					}
+					
+					MenuBuilder.AddMenuEntry(
+					   FText::FromName(DerivedClass->GetFName()),
+					   FText::GetEmpty(),
+					   FSlateIcon(),
+					   FUIAction(
+						   FExecuteAction::CreateSP(MainPanelModel.Pin().ToSharedRef(), &FObjectMixerEditorMainPanel::SetObjectFilterClass, DerivedClass),
+						   FCanExecuteAction::CreateLambda([](){ return true; }),
+						   FIsActionChecked::CreateSP(MainPanelModel.Pin().ToSharedRef(), &FObjectMixerEditorMainPanel::IsClassSelected, DerivedClass)
+					   ),
+					   NAME_None,
+					   EUserInterfaceActionType::RadioButton
+				   );
+				}
+			}
+		}
+		MenuBuilder.EndSection();
+	}
+	else
+	{
+		MenuBuilder.AddMenuEntry(LOCTEXT("NoFilterClassesAvailable", "No filter classes available."), FText::GetEmpty(), FSlateIcon(), FUIAction());
+	}
+
+	TSharedRef<SWidget> Widget = MenuBuilder.MakeWidget();
+	FChildren* ChildWidgets = Widget->GetChildren();
+	for (int32 ChildItr = 0; ChildItr < ChildWidgets->Num(); ChildItr++)
+	{
+		const TSharedRef<SWidget>& Child = ChildWidgets->GetChildAt(ChildItr);
+
+		Child->EnableToolTipForceField(false);
+	}
+	Widget->EnableToolTipForceField(false);
+	
+	return Widget;
+}
+
+TSharedRef<SWidget> SObjectMixerEditorMainPanel::BuildShowOptionsMenu()
+{
+	FMenuBuilder ShowOptionsMenuBuilder = FMenuBuilder(true, nullptr);
+
+	// No need to select filter class from outside generic instance
+	if (MainPanelModel.Pin()->GetModuleName() == FObjectMixerEditorModule::BaseObjectMixerModuleName)
+	{
+		ShowOptionsMenuBuilder.BeginSection("ListViewOptions", LOCTEXT("FilterClassManagementSection", "Filter Class Management"));
+		{
+			// Filter Class Management Button
+			const TSharedRef<SWidget> FilterClassManagementButton =
+				SNew(SBox)
+				.Padding(8, 0)
+				[
+					SNew(SComboButton)
+					.ToolTipText(LOCTEXT("FilterClassManagementButton_Tooltip", "Select a filter class"))
+					.ContentPadding(FMargin(4, 0.5f))
+					.ComboButtonStyle(&FAppStyle::Get().GetWidgetStyle<FComboButtonStyle>("ComboButton"))
+					.OnGetMenuContent(this, &SObjectMixerEditorMainPanel::OnGenerateFilterClassMenu)
+					.ForegroundColor(FStyleColors::Foreground)
+					.MenuPlacement(EMenuPlacement::MenuPlacement_MenuRight)
+					.ButtonContent()
+					[
+						SNew(SHorizontalBox)
+
+						+ SHorizontalBox::Slot()
+						.Padding(0, 1, 4, 0)
+						.AutoWidth()
+						[
+							SNew(SImage)
+							.Image(FAppStyle::Get().GetBrush("Icons.Filter"))
+							.ColorAndOpacity(FSlateColor::UseForeground())
+						]
+
+						+ SHorizontalBox::Slot()
+						.Padding(0, 1, 0, 0)
+						.AutoWidth()
+						[
+							SNew(STextBlock)
+							.Text(LOCTEXT("FilterClassToolbarButton", "Object Filter Class"))
+						]
+					]
+				];
+
+			ShowOptionsMenuBuilder.AddWidget(FilterClassManagementButton, FText::GetEmpty());
+		}
+		ShowOptionsMenuBuilder.EndSection();
+	}
+
+	// Add List View Mode Options
+	ShowOptionsMenuBuilder.BeginSection("ListViewOptions", LOCTEXT("ListViewOptionsSection", "List View Options"));
+	{
+		// Foreach on uenum
+		const FString EnumPath = "/Script/ObjectMixerEditor.EObjectMixerTreeViewMode";
+		if (const UEnum* EnumPtr = FindObject<UEnum>(nullptr, *EnumPath, true))
+		{
+			for (int32 EnumItr = 0; EnumItr < EnumPtr->GetMaxEnumValue(); EnumItr++)
+			{
+				EObjectMixerTreeViewMode EnumValue = (EObjectMixerTreeViewMode)EnumItr;
+				
+				ShowOptionsMenuBuilder.AddMenuEntry(
+					EnumPtr->GetDisplayNameTextByIndex(EnumItr),
+					EnumPtr->GetToolTipTextByIndex(EnumItr),
+					FSlateIcon(),
+					FUIAction(
+						FExecuteAction::CreateLambda([this, EnumValue]()
+						{
+							SetTreeViewMode(EnumValue);
+						}),
+						FCanExecuteAction::CreateLambda([](){ return true; }),
+						FIsActionChecked::CreateLambda([this, EnumValue]()
+						{
+							return GetTreeViewMode() == EnumValue;
+						})
+					),
+					NAME_None,
+					EUserInterfaceActionType::RadioButton
+				);
+			}
+		}
+	}
+	ShowOptionsMenuBuilder.EndSection();
+
+	ShowOptionsMenuBuilder.BeginSection("MiscOptionsSection", LOCTEXT("MiscOptionsSection","Misc"));
+	{
+		// No need to open generic instance from itself
+		if (MainPanelModel.Pin()->GetModuleName() != FObjectMixerEditorModule::BaseObjectMixerModuleName)
+		{
+			ShowOptionsMenuBuilder.AddMenuEntry(
+			  LOCTEXT("OpenGenericInstanceMenuOption", "Open Generic Object Mixer Instance"),
+			  LOCTEXT("OpenGenericInstanceMenuOptionTooltip", "Open a generic object mixer instance that can take in a user-specified filter class."),
+			  FSlateIcon(),
+			  FUIAction(FExecuteAction::CreateLambda([]()
+			  {
+				  FGlobalTabmanager::Get()->TryInvokeTab(FObjectMixerEditorModule::Get().GetTabSpawnerId());
+			  })));
+		}
+
+		ShowOptionsMenuBuilder.AddMenuEntry(
+			LOCTEXT("RebuildListMenuOption","Rebuild List"), 
+			LOCTEXT("RebuildListMenuOptionTooltip","Force the active list to be rebuilt. Useful if the list doesn't automatically refresh."),
+			FSlateIcon(),
+			FUIAction(FExecuteAction::CreateSP(MainPanelModel.Pin().ToSharedRef(), &FObjectMixerEditorMainPanel::RequestRebuildList)));
+	}
+	ShowOptionsMenuBuilder.EndSection();
+
+	if (ShowFilters.Num())
+	{
+		ShowOptionsMenuBuilder.BeginSection("", LOCTEXT("ShowOptions_ShowSectionHeading", "Show"));
+		{
+			// Add show filters
+			auto AddFiltersLambda = [this, &ShowOptionsMenuBuilder](const TSharedRef<IObjectMixerEditorListFilter>& InFilter)
+			{
+				const FString& FilterName = InFilter->GetFilterName();
+			
+				ShowOptionsMenuBuilder.AddMenuEntry(
+				   InFilter->GetFilterButtonLabel(),
+				   InFilter->GetFilterButtonToolTip(),
+				   FSlateIcon(),
+				   FUIAction(
+					   FExecuteAction::CreateLambda(
+						   [this, FilterName]()
+						   {
+							   ToggleFilterActive(FilterName);
+						   }
+						),
+					   FCanExecuteAction(),
+					   FIsActionChecked::CreateSP( InFilter, &IObjectMixerEditorListFilter::GetIsFilterActive )
+				   ),
+				   NAME_None,
+				   EUserInterfaceActionType::ToggleButton
+			   );
+			};
+
+			for (const TSharedRef<IObjectMixerEditorListFilter>& Filter : ShowFilters)
+			{
+				// Don't add non-toggle filters to Show Options, they're always on
+				if (Filter->IsToggleable())
+				{
+					AddFiltersLambda(Filter);
+				}
+			}
+		}
+		ShowOptionsMenuBuilder.EndSection();
+	}
+
+	return ShowOptionsMenuBuilder.MakeWidget();
+}
+
+void SObjectMixerEditorMainPanel::OnSearchTextChanged(const FText& Text)
+{
+	ExecuteListViewSearchOnAllRows(Text.ToString(), true);
+}
+
+FString SObjectMixerEditorMainPanel::GetSearchStringFromSearchInputField() const
+{
+	return ensureAlwaysMsgf(SearchBoxPtr.IsValid(),
+		TEXT("%hs: SearchBoxPtr is not valid. Check to make sure it was created."), __FUNCTION__)
+	? SearchBoxPtr->GetText().ToString() : "";
+}
+
+void SObjectMixerEditorMainPanel::SetSearchStringInSearchInputField(const FString InSearchString) const
+{
+	if (ensureAlwaysMsgf(SearchBoxPtr.IsValid(),
+		TEXT("%hs: SearchBoxPtr is not valid. Check to make sure it was created."), __FUNCTION__))
+	{
+		SearchBoxPtr->SetText(FText::FromString(InSearchString));
+	}
+}
+
+void SObjectMixerEditorMainPanel::ExecuteListViewSearchOnAllRows(const FString& SearchString,
+	const bool bShouldRefreshAfterward)
+{
+	MainPanelModel.Pin()->GetEditorListModel().Pin()->ExecuteListViewSearchOnAllRows(GetSearchStringFromSearchInputField());
+}
+
+EObjectMixerTreeViewMode SObjectMixerEditorMainPanel::GetTreeViewMode()
+{
+	const TSharedPtr<FObjectMixerEditorMainPanel> PinnedPanelModel = MainPanelModel.Pin();
+	check(PinnedPanelModel);
+	
+	return PinnedPanelModel->GetTreeViewMode();
+}
+
+void SObjectMixerEditorMainPanel::SetTreeViewMode(EObjectMixerTreeViewMode InViewMode)
+{
+	if (TSharedPtr<FObjectMixerEditorMainPanel> PinnedPanelModel = MainPanelModel.Pin())
+	{
+		PinnedPanelModel->SetTreeViewMode(InViewMode);
+	}
+}
+
+void SObjectMixerEditorMainPanel::ToggleFilterActive(const FString& FilterName)
+{
+	if (const TSharedRef<IObjectMixerEditorListFilter>* Match =
+		Algo::FindByPredicate(GetShowFilters(),
+		[&FilterName](TSharedRef<IObjectMixerEditorListFilter> Comparator)
+		{
+			return Comparator->GetFilterName().Equals(FilterName);
+		}))
+	{
+		const TSharedRef<IObjectMixerEditorListFilter> Filter = *Match;
+		Filter->ToggleFilterActive();
+
+		MainPanelModel.Pin()->GetEditorListModel().Pin()->EvaluateIfRowsPassFilters();
+	}
+}
+
+const TSet<FName>& SObjectMixerEditorMainPanel::GetCurrentCategorySelection()
+{
+	return CurrentCategorySelection;
+}
+
+void SObjectMixerEditorMainPanel::RebuildCategorySelector()
+{
+	check(MainPanelModel.IsValid());
+	
+	CategorySelectorBox->ClearChildren();
+	CategorySelectorBox->SetVisibility(EVisibility::Collapsed);
+
+	TArray<FName> AllCategories = MainPanelModel.Pin()->GetAllCategories().Array();
+	
+	if (AllCategories.IsEmpty())
+	{
+		// we've selected something that has no sections - rather than show just "All", hide the box
+		ResetCurrentCategorySelection();
+		return;
+	}
+
+	auto CreateSection = [this](FName CategoryName) 
+	{
+		return SNew(SBox)
+			.Padding(FMargin(0))
+			[
+				SNew(SCheckBox)
+				.Style(FAppStyle::Get(), "DetailsView.SectionButton")
+				.OnCheckStateChanged(this, &SObjectMixerEditorMainPanel::OnCategoryCheckedChanged, CategoryName)
+				.IsChecked(this, &SObjectMixerEditorMainPanel::IsCategoryChecked, CategoryName)
+				[
+					SNew(STextBlock)
+					.TextStyle(FAppStyle::Get(), "SmallText")
+					.Text(FText::FromName(CategoryName))
+				]
+			];
+	};
+	
+	AllCategories.Sort([](FName A, FName B)
+		{
+			return A.LexicalLess(B);
+		});
+
+	for (const FName& Key : AllCategories)
+	{
+		CategorySelectorBox->AddSlot()
+		[
+			CreateSection(Key)
+		];
+	}
+
+	CategorySelectorBox->AddSlot()
+	[
+		CreateSection("All")
 	];
+
+	CategorySelectorBox->SetVisibility(EVisibility::Visible);
+}
+
+void SObjectMixerEditorMainPanel::OnCategoryCheckedChanged(ECheckBoxState State, FName SectionName)
+{
+	check(MainPanelModel.IsValid());
+
+	// Remove category
+	if (FSlateApplication::Get().GetModifierKeys().IsAltDown())
+	{
+		MainPanelModel.Pin()->RemoveCategory(SectionName);
+
+		CurrentCategorySelection.Remove(SectionName);
+	}
+	else
+	{
+		const bool bIsControlDown = FSlateApplication::Get().GetModifierKeys().IsControlDown();
+
+		if (State == ECheckBoxState::Unchecked)
+		{
+			if (bIsControlDown)
+			{
+				CurrentCategorySelection.Remove(SectionName);
+			}
+			else
+			{
+				CurrentCategorySelection.Reset();
+
+				if (SectionName != "All")
+				{
+					CurrentCategorySelection.Add(SectionName);
+				}
+			}
+		}
+		else if (State == ECheckBoxState::Checked)
+		{
+			if (!bIsControlDown)
+			{
+				CurrentCategorySelection.Reset();
+			}
+
+			if (SectionName != "All")
+			{
+				CurrentCategorySelection.Add(SectionName);
+			}
+		}
+	}
+
+	MainPanelModel.Pin()->GetEditorListModel().Pin()->EvaluateIfRowsPassFilters();
+}
+
+ECheckBoxState SObjectMixerEditorMainPanel::IsCategoryChecked(FName Section) const
+{
+	if (CurrentCategorySelection.IsEmpty() && Section == "All")
+	{
+		return ECheckBoxState::Checked;
+	}
+
+	return CurrentCategorySelection.Contains(Section) ? ECheckBoxState::Checked : ECheckBoxState::Unchecked;
 }
 
 SObjectMixerEditorMainPanel::~SObjectMixerEditorMainPanel()
 {
-	MainPanel.Reset();
+	if (MainPanelModel.IsValid())
+	{
+		MainPanelModel.Pin()->GetOnObjectMixerCategoryMapChanged().RemoveAll(this);
+		MainPanelModel.Reset();
+	}
 }
 
 #undef LOCTEXT_NAMESPACE
diff --git a/Engine/Source/Editor/ObjectMixer/ObjectMixer/Private/Views/Widgets/ObjectMixerEditorListMenuContext.cpp b/Engine/Source/Editor/ObjectMixer/ObjectMixer/Private/Views/Widgets/ObjectMixerEditorListMenuContext.cpp
new file mode 100644
--- /dev/null
+++ b/Engine/Source/Editor/ObjectMixer/ObjectMixer/Private/Views/Widgets/ObjectMixerEditorListMenuContext.cpp
@@ -0,0 +1,239 @@
+// Copyright Epic Games, Inc. All Rights Reserved.
+
+#include "Views/Widgets/ObjectMixerEditorListMenuContext.h"
+
+#include "Framework/MultiBox/MultiBoxDefs.h"
+#include "ToolMenu.h"
+#include "ToolMenus.h"
+#include "ToolMenuContext.h"
+#include "ToolMenuDelegates.h"
+#include "ToolMenuEntry.h"
+#include "ToolMenuSection.h"
+#include "UnrealEdGlobals.h"
+#include "Views/MainPanel/ObjectMixerEditorMainPanel.h"
+#include "Widgets/Layout/SBox.h"
+#include "Widgets/Input/SEditableTextBox.h"
+
+#define LOCTEXT_NAMESPACE "ObjectMixerEditorListMenuContext"
+
+FName UObjectMixerEditorListMenuContext::DefaultContextBaseMenuName("ObjectMixer.ContextMenuBase");
+
+TSharedPtr<SWidget> UObjectMixerEditorListMenuContext::CreateContextMenu(const FObjectMixerEditorListMenuContextData InData)
+{
+	if (InData.SelectedItems.Num() == 0)
+	{
+		return nullptr;
+	}
+
+	return BuildContextMenu(InData);
+}
+
+TSharedPtr<SWidget> UObjectMixerEditorListMenuContext::BuildContextMenu(const FObjectMixerEditorListMenuContextData& InData)
+{
+	RegisterContextMenu();
+
+	UObjectMixerEditorListMenuContext* ContextObject = NewObject<UObjectMixerEditorListMenuContext>();
+	ContextObject->Data = InData;
+
+	const FToolMenuContext Context(ContextObject);
+
+	// Build up the menu for a selection
+	UToolMenus* ToolMenus = UToolMenus::Get();
+	UToolMenu* Menu = ToolMenus->GenerateMenu(DefaultContextBaseMenuName, Context);
+
+	for (const FToolMenuSection& Section : Menu->Sections)
+	{
+		if (Section.Blocks.Num() > 0)
+		{
+			return ToolMenus->GenerateWidget(Menu);
+		}
+	}
+
+	return nullptr;
+}
+
+void UObjectMixerEditorListMenuContext::RegisterContextMenu()
+{
+	UToolMenus* ToolMenus = UToolMenus::Get();
+	
+	if (!ToolMenus->IsMenuRegistered(DefaultContextBaseMenuName))
+	{
+		UToolMenu* Menu = ToolMenus->RegisterMenu(DefaultContextBaseMenuName);
+	
+		Menu->AddDynamicSection("DynamicCategorizationSection", FNewToolMenuDelegate::CreateLambda(
+			[](UToolMenu* InMenu)
+			{
+				UObjectMixerEditorListMenuContext* Context = InMenu->FindContext<UObjectMixerEditorListMenuContext>();
+				if (!Context || Context->Data.SelectedItems.Num() == 0)
+				{
+					return;
+				}
+	
+				FToolMenuSection& Section = InMenu->FindOrAddSection("Categorization");
+				Section.Label = LOCTEXT("CategorizationSectionName", "Categorization");
+			
+				Section.AddSubMenu(
+					"SelectCategorySubMenu",
+					LOCTEXT("SelectCategorySubmenu", "Select or Add Category"),
+					LOCTEXT("SelectCategorySubmenu_Tooltip", "Select the category to which you wish to assign this object."),
+					FNewToolMenuDelegate::CreateLambda(
+						[Context](UToolMenu* Menu)
+						{
+							FToolMenuEntry Args;
+							Args.Type = EMultiBlockType::Widget;
+							Args.MakeCustomWidget.BindLambda(
+								[Context](const FToolMenuContext&, const FToolMenuCustomWidgetContext&)
+								{
+									return SNew(SBox)
+											.MinDesiredWidth(200)
+											.Padding(8, 0)
+											[
+												SAssignNew(Context->EditableText, SEditableTextBox)
+												.HintText(LOCTEXT("NewCategoryEditableHintText", "Enter a new category name..."))
+												.OnTextCommitted_Static(&UObjectMixerEditorListMenuContext::OnTextCommitted, Context)
+												.OnTextChanged_Static(&UObjectMixerEditorListMenuContext::OnTextChanged, Context)
+											]
+									;
+								}
+							);
+							Menu->AddMenuEntry("NewCategoryInput", Args);
+							
+							if (const TSharedPtr<FObjectMixerEditorMainPanel> MainPanel = Context->Data.MainPanelPtr.Pin())
+							{
+								if (TArray<FName> Categories = MainPanel->GetAllCategories().Array(); Categories.Num() > 0)
+								{
+									FToolMenuSection& Section = Menu->FindOrAddSection("Categories");
+								   Section.Label = LOCTEXT("CategoriesSectionName", "Categories");
+
+								   Categories.StableSort([](const FName& A, const FName B)
+								   {
+									   return A.LexicalLess(B);
+								   });
+								
+								   for (const FName& Key : Categories)
+								   {
+									   const FText KeyText = FText::FromName(Key);
+									
+									   Section.AddMenuEntry(
+										   Key,
+										   KeyText,
+										   FText::Format(LOCTEXT("AddObjectsToCategoryTooltipFormat", "Add selected to category '{0}'"), KeyText),
+										   FSlateIcon(),
+										   FUIAction(
+											   FExecuteAction::CreateStatic(&UObjectMixerEditorListMenuContext::OnClickCategoryMenuEntry, Key, Context),
+											   FCanExecuteAction(),
+											   FIsActionChecked::CreateStatic(&UObjectMixerEditorListMenuContext::AreAllObjectsInCategory, Key, Context)
+										   ),
+										   EUserInterfaceActionType::Check
+									   );
+								   }
+								}
+							}
+						}
+					)
+				);
+			}
+		));
+	}
+}
+
+void UObjectMixerEditorListMenuContext::OnTextCommitted(const FText& InText, ETextCommit::Type InCommitType, UObjectMixerEditorListMenuContext* Context)
+{
+	check(Context);
+	
+	if (InCommitType == ETextCommit::OnEnter)
+	{
+		UE_LOG(LogTemp, Warning, TEXT("%hs: %s"), __FUNCTION__, *InText.ToString());
+
+		AddObjectsToCategory(*InText.ToString(), Context);
+	}
+}
+
+void UObjectMixerEditorListMenuContext::OnTextChanged(const FText& InText, UObjectMixerEditorListMenuContext* Context)
+{
+	check(Context);
+	
+	if (const FString AsString = InText.ToString(); AsString.Len() > 12)
+	{
+		Context->EditableText->SetText(FText::FromString(AsString.LeftChop(1)));
+	}
+}
+
+void UObjectMixerEditorListMenuContext::OnClickCategoryMenuEntry(const FName Key, UObjectMixerEditorListMenuContext* Context)
+{
+	if (AreAllObjectsInCategory(Key, Context))
+	{
+		RemoveObjectsFromCategory(Key, Context);
+	}
+	else
+	{
+		AddObjectsToCategory(Key, Context);
+	}
+}
+
+void UObjectMixerEditorListMenuContext::AddObjectsToCategory(const FName Key, UObjectMixerEditorListMenuContext* Context)
+{
+	check(Context);
+
+	if (const TSharedPtr<FObjectMixerEditorMainPanel> MainPanel = Context->Data.MainPanelPtr.Pin())
+	{
+		TSet<FSoftObjectPath> ObjectPaths;
+
+		for (const TSharedPtr<FObjectMixerEditorListRow>& Item : Context->Data.SelectedItems)
+		{
+			if (const UObject* Object = Item->GetObject())
+			{
+				ObjectPaths.Add(Object);
+			}
+		}
+		
+		MainPanel->AddObjectsToCategory(Key, ObjectPaths);
+	}
+}
+
+void UObjectMixerEditorListMenuContext::RemoveObjectsFromCategory(const FName Key, UObjectMixerEditorListMenuContext* Context)
+{
+	check(Context);
+
+	if (const TSharedPtr<FObjectMixerEditorMainPanel> MainPanel = Context->Data.MainPanelPtr.Pin())
+	{
+		TSet<FSoftObjectPath> ObjectPaths;
+
+		for (const TSharedPtr<FObjectMixerEditorListRow>& Item : Context->Data.SelectedItems)
+		{
+			if (const UObject* Object = Item->GetObject())
+			{
+				ObjectPaths.Add(Object);
+			}
+		}
+		
+		MainPanel->RemoveObjectsFromCategory(Key, ObjectPaths);
+	}
+}
+
+bool UObjectMixerEditorListMenuContext::AreAllObjectsInCategory(const FName Key, UObjectMixerEditorListMenuContext* Context)
+{
+	check(Context);
+
+	bool bAreAllSelectedObjectsInCategory = false;
+
+	if (const TSharedPtr<FObjectMixerEditorMainPanel> MainPanel = Context->Data.MainPanelPtr.Pin())
+	{
+		for (const TSharedPtr<FObjectMixerEditorListRow>& Item : Context->Data.SelectedItems)
+		{
+			if (const UObject* Object = Item->GetObject())
+			{
+				bAreAllSelectedObjectsInCategory = MainPanel->IsObjectInCategory(Key, Object);
+
+				if (!bAreAllSelectedObjectsInCategory)
+				{
+					return false;
+				}
+			}
+		}
+	}
+
+	return bAreAllSelectedObjectsInCategory;
+}
+
+#undef LOCTEXT_NAMESPACE
diff --git a/Engine/Source/Editor/ObjectMixer/ObjectMixer/Public/ObjectMixerEditorModule.h b/Engine/Source/Editor/ObjectMixer/ObjectMixer/Public/ObjectMixerEditorModule.h
--- a/Engine/Source/Editor/ObjectMixer/ObjectMixer/Public/ObjectMixerEditorModule.h
+++ b/Engine/Source/Editor/ObjectMixer/ObjectMixer/Public/ObjectMixerEditorModule.h
@@ -23,6 +23,8 @@ public:
 	void AddOnComponentEditedDelegate(const FDelegateHandle& InOnComponentEditedHandle);
 
 	static FObjectMixerEditorModule& Get();
+
+	virtual FName GetModuleName();
 	
 	virtual TSharedPtr<SWidget> MakeObjectMixerDialog() const;
 
@@ -57,6 +59,8 @@ public:
 
 	TSharedPtr<FWorkspaceItem> GetWorkspaceGroup();
 
+	const static FName BaseObjectMixerModuleName;
+
 protected:
 
 	virtual TSharedRef<SDockTab> SpawnMainPanelTab();
diff --git a/Engine/Source/Editor/ObjectMixer/ObjectMixer/Public/ObjectMixerEditorProjectSettings.h b/Engine/Source/Editor/ObjectMixer/ObjectMixer/Public/ObjectMixerEditorProjectSettings.h
--- a/Engine/Source/Editor/ObjectMixer/ObjectMixer/Public/ObjectMixerEditorProjectSettings.h
+++ b/Engine/Source/Editor/ObjectMixer/ObjectMixer/Public/ObjectMixerEditorProjectSettings.h
@@ -2,8 +2,54 @@
 
 #pragma once
 
+#include "Algo/Find.h"
+
 #include "ObjectMixerEditorProjectSettings.generated.h"
 
+USTRUCT()
+struct FObjectMixerCategoryObjectSet
+{
+	GENERATED_BODY()
+	
+	UPROPERTY()
+	FName CategoryName = NAME_None;
+
+	UPROPERTY()
+	TSet<FSoftObjectPath> CategoryObjects = {};
+
+	bool operator==(const FObjectMixerCategoryObjectSet& Other) const
+	{
+		return CategoryName.IsEqual(Other.CategoryName);
+	}
+
+	friend uint32 GetTypeHash (const FObjectMixerCategoryObjectSet& Other)
+	{
+		return GetTypeHash(Other.CategoryName);
+	}
+};
+
+USTRUCT()
+struct FObjectMixerSerializationData
+{
+	GENERATED_BODY()
+	
+	UPROPERTY()
+	FName FilterClassName = NAME_None;
+
+	UPROPERTY()
+	TSet<FObjectMixerCategoryObjectSet> SerializedCategories = {};
+
+	bool operator==(const FObjectMixerSerializationData& Other) const
+	{
+		return FilterClassName.IsEqual(Other.FilterClassName);
+	}
+
+	friend uint32 GetTypeHash (const FObjectMixerSerializationData& Other)
+	{
+		return GetTypeHash(Other.FilterClassName);
+	}
+};
+
 UCLASS(config = Engine, defaultconfig)
 class OBJECTMIXEREDITOR_API UObjectMixerEditorProjectSettings : public UObject
 {
@@ -18,4 +64,143 @@ public:
 	 */
 	UPROPERTY(Config, EditAnywhere, Category="Object Mixer")
 	bool bExpandTreeViewItemsByDefault = true;
+
+	UPROPERTY(Config)
+	TSet<FObjectMixerSerializationData> SerializedData;
+
+	FObjectMixerSerializationData* FindSerializationDataByFilterClassName(const FName& FilterClassName)
+	{
+		if (FObjectMixerSerializationData* Match = Algo::FindByPredicate(SerializedData,
+			[FilterClassName](const FObjectMixerSerializationData& Comparator)
+			{
+				return Comparator.FilterClassName.IsEqual(FilterClassName);
+			}))
+		{
+			return Match;
+		}
+
+		return nullptr;
+	}
+
+	void AddObjectsToCategory(const FName& FilterClassName, const FName& CategoryName, const TSet<FSoftObjectPath>& ObjectsToAdd)
+	{
+		FObjectMixerSerializationData* Data = FindSerializationDataByFilterClassName(FilterClassName);
+
+		if (!Data)
+		{
+			SerializedData.Add({FilterClassName});
+			Data = FindSerializationDataByFilterClassName(FilterClassName);
+		}
+		
+		if (Data)
+		{
+			if (FObjectMixerCategoryObjectSet* Match = Algo::FindByPredicate(Data->SerializedCategories,
+				[CategoryName](const FObjectMixerCategoryObjectSet& Comparator)
+				{
+					return Comparator.CategoryName.IsEqual(CategoryName);
+				}))
+			{
+				(*Match).CategoryObjects.Append(ObjectsToAdd);
+			}
+			else
+			{
+				Data->SerializedCategories.Add({CategoryName, ObjectsToAdd});
+			}
+
+			SaveConfig();
+		}
+	}
+
+	void RemoveObjectsFromCategory(const FName& FilterClassName, const FName& CategoryName, const TSet<FSoftObjectPath>& ObjectsToRemove)
+	{
+		if (FObjectMixerSerializationData* Data = FindSerializationDataByFilterClassName(FilterClassName))
+		{
+			if (FObjectMixerCategoryObjectSet* Match = Algo::FindByPredicate(Data->SerializedCategories,
+				[CategoryName](const FObjectMixerCategoryObjectSet& Comparator)
+				{
+					return Comparator.CategoryName.IsEqual(CategoryName);
+				}))
+			{
+				for (const FSoftObjectPath& ObjectPath : ObjectsToRemove)
+				{
+					(*Match).CategoryObjects.Remove(ObjectPath);
+				}
+				
+				SaveConfig();
+			}
+		}
+	}
+
+	void RemoveCategory(const FName& FilterClassName, const FName& CategoryName)
+	{
+		if (FObjectMixerSerializationData* Data = FindSerializationDataByFilterClassName(FilterClassName))
+		{
+			if (FObjectMixerCategoryObjectSet* Match = Algo::FindByPredicate(Data->SerializedCategories,
+				[CategoryName](const FObjectMixerCategoryObjectSet& Comparator)
+				{
+					return Comparator.CategoryName.IsEqual(CategoryName);
+				}))
+			{
+				(*Match).CategoryObjects.Empty();
+
+				Data->SerializedCategories.Remove(*Match);
+
+				SaveConfig();
+			}
+		}
+	}
+
+	bool IsObjectInCategory(const FName& FilterClassName, const FName& CategoryName, const FSoftObjectPath& InObject)
+	{
+		if (FObjectMixerSerializationData* Data = FindSerializationDataByFilterClassName(FilterClassName))
+		{
+			if (const FObjectMixerCategoryObjectSet* Match = Algo::FindByPredicate(Data->SerializedCategories,
+				[CategoryName](const FObjectMixerCategoryObjectSet& Comparator)
+				{
+					return Comparator.CategoryName.IsEqual(CategoryName);
+				}))
+			{
+				return (*Match).CategoryObjects.Contains(InObject);
+			}
+		}
+
+		return false;
+	}
+
+	TSet<FName> GetCategoriesForObject(const FName& FilterClassName, const FSoftObjectPath& InObject)
+	{
+		if (FObjectMixerSerializationData* Data = FindSerializationDataByFilterClassName(FilterClassName))
+		{
+			TSet<FName> CategoriesWithObject;
+
+			for (const FObjectMixerCategoryObjectSet& SerializedCategory : Data->SerializedCategories)
+			{
+				if (IsObjectInCategory(FilterClassName, SerializedCategory.CategoryName, InObject))
+				{
+					CategoriesWithObject.Add(SerializedCategory.CategoryName);
+				}
+			}
+
+			return CategoriesWithObject;
+		}
+
+		return {};
+	}
+
+	TSet<FName> GetAllCategories(const FName& FilterClassName)
+	{
+		if (FObjectMixerSerializationData* Data = FindSerializationDataByFilterClassName(FilterClassName))
+		{
+			TSet<FName> Categories;
+
+			for (const FObjectMixerCategoryObjectSet& SerializedCategory : Data->SerializedCategories)
+			{
+				Categories.Add(SerializedCategory.CategoryName);
+			}
+
+			return Categories;
+		}
+
+		return {};
+	}
 };
diff --git a/Engine/Source/Editor/ObjectMixer/ObjectMixer/Public/Views/List/ObjectMixerEditorList.h b/Engine/Source/Editor/ObjectMixer/ObjectMixer/Public/Views/List/ObjectMixerEditorList.h
--- a/Engine/Source/Editor/ObjectMixer/ObjectMixer/Public/Views/List/ObjectMixerEditorList.h
+++ b/Engine/Source/Editor/ObjectMixer/ObjectMixer/Public/Views/List/ObjectMixerEditorList.h
@@ -2,101 +2,27 @@
 
 #pragma once
 
-#include "ObjectFilter/ObjectMixerEditorObjectFilter.h"
+#include "Views/MainPanel/ObjectMixerEditorMainPanel.h"
 
-#include "UObject/StrongObjectPtr.h"
+#include "Templates/SharedPointer.h"
 #include "Widgets/SWidget.h"
 
-struct FObjectMixerEditorListRow;
-typedef TSharedPtr<FObjectMixerEditorListRow> FObjectMixerEditorListRowPtr;
-
 class SObjectMixerEditorList;
 
 class OBJECTMIXEREDITOR_API FObjectMixerEditorList : public TSharedFromThis<FObjectMixerEditorList>
 {
 public:
 
-	FObjectMixerEditorList() = default;
+	FObjectMixerEditorList(TSharedRef<FObjectMixerEditorMainPanel, ESPMode::ThreadSafe> InMainPanel)
+	: MainPanelModelPtr(InMainPanel)
+	{}
 
 	virtual ~FObjectMixerEditorList();
 
 	void FlushWidget();
+	
 	TSharedRef<SWidget> GetOrCreateWidget();
 
-	UObjectMixerObjectFilter* GetObjectFilter();
-
-	void CacheObjectFilterObject();
-
-	/**
-	 * Get the style of the tree (flat list or hierarchy)
-	 */
-	EObjectMixerTreeViewMode GetTreeViewMode()
-	{
-		return TreeViewMode;
-	}
-	/**
-	 * Set the style of the tree (flat list or hierarchy)
-	 */
-	void SetTreeViewMode(EObjectMixerTreeViewMode InViewMode)
-	{
-		TreeViewMode = InViewMode;
-		RequestRebuildList();
-	}
-
-	/**
-	 * Returns result from Filter->GetObjectClassesToFilter.
-	 */
-	TSet<UClass*> GetObjectClassesToFilter()
-	{
-		if (const UObjectMixerObjectFilter* Filter = GetObjectFilter())
-		{
-			return Filter->GetObjectClassesToFilter();
-		}
-		
-		return {};
-	}
-
-	/**
-	 * Returns result from Filter->GetObjectClassesToPlace.
-	 */
-	TSet<TSubclassOf<AActor>> GetObjectClassesToPlace()
-	{
-		TSet<TSubclassOf<AActor>> ReturnValue;
-
-		if (const UObjectMixerObjectFilter* Filter = GetObjectFilter())
-		{
-			ReturnValue = Filter->GetObjectClassesToPlace();
-		}
-		
-		return ReturnValue;
-	}
-
-	/**
-	 * Get the row that has solo visibility. All other rows should be set to temporarily invisible in editor.
-	 */
-	TWeakPtr<FObjectMixerEditorListRow> GetSoloRow()
-	{
-		return SoloRow;
-	}
-
-	/**
-	 * Set the row that has solo visibility. This does not set temporary editor invisibility for other rows.
-	 */
-	void SetSoloRow(TSharedRef<FObjectMixerEditorListRow> InRow)
-	{
-		SoloRow = InRow;
-	}
-
-	/**
-	 * Clear the row that has solo visibility. This does not remove temporary editor invisibility for other rows.
-	 */
-	void ClearSoloRow()
-	{
-		SoloRow = nullptr;
-	}
-
-	void SetSearchString(const FString& SearchString);
-
 	void ClearList() const;
 
 	/**
@@ -110,46 +36,21 @@ public:
 	 */
 	void RefreshList() const;
 
-	bool IsClassSelected(UClass* InNewClass) const;
+	void ExecuteListViewSearchOnAllRows(const FString& SearchString, const bool bShouldRefreshAfterward = true);
 
-	TSubclassOf<UObjectMixerObjectFilter> GetObjectFilterClass() const
-	{
-		return ObjectFilterClass;
-	}
+	void EvaluateIfRowsPassFilters(const bool bShouldRefreshAfterward = true) const;
 
-	void SetObjectFilterClass(UClass* InObjectFilterClass)
-	{
-		if (ensureAlwaysMsgf(InObjectFilterClass->IsChildOf(UObjectMixerObjectFilter::StaticClass()), TEXT("%hs: Class '%s' is not a child of UObjectMixerObjectFilter."), __FUNCTION__, *InObjectFilterClass->GetName()))
-		{
-			ObjectFilterClass = InObjectFilterClass;
-			CacheObjectFilterObject();
-			RequestRebuildList();
-		}
-	}
+	TWeakPtr<FObjectMixerEditorMainPanel> GetMainPanelModel();
 
-private:
+	TWeakPtr<FObjectMixerEditorListRow> GetSoloRow();
 
-	/**
-	 * The class used to generate property edit columns
-	 */
-	TSubclassOf<UObjectMixerObjectFilter> ObjectFilterClass;
-
-	/**
-	 * Determines the style of the tree (flat list or hierarchy)
-	 */
-	EObjectMixerTreeViewMode TreeViewMode = EObjectMixerTreeViewMode::FolderObjectSubObject;
+	void SetSoloRow(TSharedRef<FObjectMixerEditorListRow> InRow);
 
-	TSharedPtr<SObjectMixerEditorList> ListWidget;
+	void ClearSoloRow();
 
-	TStrongObjectPtr<UObjectMixerObjectFilter> ObjectFilterPtr;
+private:
 
-	TWeakPtr<FObjectMixerEditorListRow> SoloRow = nullptr;
+	TWeakPtr<FObjectMixerEditorMainPanel> MainPanelModelPtr;
 
-	// Delegates
-	TFunction<void()> RebuildDelegate = [this]()
-	{
-		RequestRebuildList();
-	};
-	
-	TSet<FDelegateHandle> DelegateHandles;
+	TSharedPtr<SObjectMixerEditorList> ListWidget;
 };
diff --git a/Engine/Source/Editor/ObjectMixer/ObjectMixer/Public/Views/List/ObjectMixerEditorListFilters/IObjectMixerEditorListFilter.h b/Engine/Source/Editor/ObjectMixer/ObjectMixer/Public/Views/List/ObjectMixerEditorListFilters/IObjectMixerEditorListFilter.h
--- a/Engine/Source/Editor/ObjectMixer/ObjectMixer/Public/Views/List/ObjectMixerEditorListFilters/IObjectMixerEditorListFilter.h
+++ b/Engine/Source/Editor/ObjectMixer/ObjectMixer/Public/Views/List/ObjectMixerEditorListFilters/IObjectMixerEditorListFilter.h
@@ -18,12 +18,15 @@ public:
 	
 	virtual ~IObjectMixerEditorListFilter() = default;
 
-	virtual FString GetFilterName() = 0;
+	virtual FString GetFilterName() const = 0;
+
+	/** Return true if the filter should be able to be toggled on/off in the Show Options. */
+	virtual bool IsToggleable() const = 0;
 
 	/** Returns localized text to display for when this filter' active state is defined by a toggle button. */
-	virtual FText GetFilterButtonLabel() = 0;
+	virtual FText GetFilterButtonLabel() const = 0;
 	/** Returns localized text to display when the user mouses over a toggle button that defines whether this filter is active. */
-	virtual FText GetFilterButtonToolTip() = 0;
+	virtual FText GetFilterButtonToolTip() const = 0;
 	
 	void SetFilterActive(const bool bNewEnabled)
 	{
@@ -37,7 +40,7 @@ public:
 
 	bool GetIsFilterActive() const
 	{
-		return bIsFilterActive;
+		return !IsToggleable() || bIsFilterActive;
 	}
 
 	void SetFilterMatchType(const EObjectMixerEditorListFilterMatchType InType)
@@ -50,7 +53,7 @@ public:
 		return FilterMatchType;
 	}
 
-	virtual bool DoesItemPassFilter(const FObjectMixerEditorListRowPtr& InItem)
+	virtual bool DoesItemPassFilter(const FObjectMixerEditorListRowPtr& InItem) const
 	{
 		return false;
 	}
diff --git a/Engine/Source/Editor/ObjectMixer/ObjectMixer/Public/Views/List/ObjectMixerEditorListFilters/ObjectMixerEditorListFilter_Category.h b/Engine/Source/Editor/ObjectMixer/ObjectMixer/Public/Views/List/ObjectMixerEditorListFilters/ObjectMixerEditorListFilter_Category.h
new file mode 100644
--- /dev/null
+++ b/Engine/Source/Editor/ObjectMixer/ObjectMixer/Public/Views/List/ObjectMixerEditorListFilters/ObjectMixerEditorListFilter_Category.h
@@ -0,0 +1,54 @@
+// Copyright Epic Games, Inc. All Rights Reserved.
+
+#pragma once
+
+#include "IObjectMixerEditorListFilter.h"
+
+#define LOCTEXT_NAMESPACE "ObjectMixerEditor"
+
+class OBJECTMIXEREDITOR_API FObjectMixerEditorListFilter_Category : public IObjectMixerEditorListFilter
+{
+public:
+
+	FObjectMixerEditorListFilter_Category()
+	{
+		SetFilterMatchType(EObjectMixerEditorListFilterMatchType::MatchAll);
+	}
+
+	virtual FString GetFilterName() const override
+	{
+		return "ObjectMixerCategoryListFilter";
+	}
+
+	virtual bool IsToggleable() const override
+	{
+		return false;
+	}
+
+	virtual FText GetFilterButtonLabel() const override
+	{
+		return LOCTEXT("ShowSourceTextFilterFormat", "Show Categories");
+	}
+
+	virtual FText GetFilterButtonToolTip() const override
+	{
+		return LOCTEXT("ShowSourceTextFilterTooltipFormat", "Show rows that are assigned to the selected categories");
+	}
+
+	virtual bool DoesItemPassFilter(const FObjectMixerEditorListRowPtr& InItem) const override
+	{
+		if (InItem.IsValid())
+		{
+			if (InItem->GetRowType() == FObjectMixerEditorListRow::ContainerObject ||
+				InItem->GetRowType() == FObjectMixerEditorListRow::MatchingObject)
+			{
+				return InItem->IsObjectRefInSelectedCategories();
+			}
+		}
+
+		return false;
+	}
+
+};
+
+#undef LOCTEXT_NAMESPACE
diff --git a/Engine/Source/Editor/ObjectMixer/ObjectMixer/Public/Views/List/ObjectMixerEditorListFilters/ObjectMixerEditorListFilter_Source.h b/Engine/Source/Editor/ObjectMixer/ObjectMixer/Public/Views/List/ObjectMixerEditorListFilters/ObjectMixerEditorListFilter_Source.h
deleted file mode 100644
diff --git a/Engine/Source/Editor/ObjectMixer/ObjectMixer/Public/Views/List/ObjectMixerEditorListRow.h b/Engine/Source/Editor/ObjectMixer/ObjectMixer/Public/Views/List/ObjectMixerEditorListRow.h
--- a/Engine/Source/Editor/ObjectMixer/ObjectMixer/Public/Views/List/ObjectMixerEditorListRow.h
+++ b/Engine/Source/Editor/ObjectMixer/ObjectMixer/Public/Views/List/ObjectMixerEditorListRow.h
@@ -53,10 +53,9 @@ struct OBJECTMIXEREDITOR_API FObjectMixerEditorListRow final : TSharedFromThis<F
 
 	UObjectMixerObjectFilter* GetObjectFilter() const;
 
-	[[nodiscard]] EObjectMixerEditorListRowType GetRowType() const;
+	bool IsObjectRefInSelectedCategories() const;
 
-	[[nodiscard]] int32 GetChildDepth() const;
-	void SetChildDepth(const int32 InDepth);
+	[[nodiscard]] EObjectMixerEditorListRowType GetRowType() const;
 
 	[[nodiscard]] int32 GetSortOrder() const;
 	void SetSortOrder(const int32 InNewOrder);
@@ -97,10 +96,7 @@ struct OBJECTMIXEREDITOR_API FObjectMixerEditorListRow final : TSharedFromThis<F
 	[[nodiscard]] bool ShouldRowWidgetBeVisible() const;
 	[[nodiscard]] EVisibility GetDesiredRowWidgetVisibility() const;
 
-	[[nodiscard]] bool HasVisibleChildRowWidgets() const
-	{
-		return false;
-	}
+	[[nodiscard]] bool HasVisibleChildRowWidgets() const;
 
 	[[nodiscard]] FText GetDisplayName();
 
@@ -123,10 +119,6 @@ struct OBJECTMIXEREDITOR_API FObjectMixerEditorListRow final : TSharedFromThis<F
 	 * Determines the style of the tree (flat list or hierarchy)
 	 */
 	EObjectMixerTreeViewMode GetTreeViewMode();
-	/**
-	 * Determine the style of the tree (flat list or hierarchy)
-	 */
-	void SetTreeViewMode(EObjectMixerTreeViewMode InViewMode);
 
 	[[nodiscard]] TArray<FObjectMixerEditorListRowPtr> GetSelectedTreeViewItems() const;
 
@@ -151,10 +143,6 @@ private:
 	TWeakPtr<SObjectMixerEditorList> ListViewPtr;
 
 	FText DisplayNameOverride;
-	
-	bool bIsTreeViewItemExpanded = false;
-
-	int32 ChildDepth = 0;
 
 	int32 SortOrder = -1;
 
diff --git a/Engine/Source/Editor/ObjectMixer/ObjectMixer/Public/Views/List/SObjectMixerEditorList.h b/Engine/Source/Editor/ObjectMixer/ObjectMixer/Public/Views/List/SObjectMixerEditorList.h
--- a/Engine/Source/Editor/ObjectMixer/ObjectMixer/Public/Views/List/SObjectMixerEditorList.h
+++ b/Engine/Source/Editor/ObjectMixer/ObjectMixer/Public/Views/List/SObjectMixerEditorList.h
@@ -113,23 +113,13 @@ public:
 		return TreeViewRootObjects.Num();
 	}
 	
-	TWeakPtr<FObjectMixerEditorListRow> GetSoloRow()
-	{
-		return GetListModelPtr().Pin()->GetSoloRow();
-	}
+	TWeakPtr<FObjectMixerEditorListRow> GetSoloRow();
 
-	void SetSoloRow(TSharedRef<FObjectMixerEditorListRow> InRow)
-	{
-		GetListModelPtr().Pin()->SetSoloRow(InRow);
-	}
+	void SetSoloRow(TSharedRef<FObjectMixerEditorListRow> InRow);
 
-	void ClearSoloRow()
-	{
-		GetListModelPtr().Pin()->ClearSoloRow();
-	}
+	void ClearSoloRow();
 
 	FString GetSearchStringFromSearchInputField() const;
-	void SetSearchStringInSearchInputField(const FString InSearchString) const;
 	void ExecuteListViewSearchOnAllRows(const FString& SearchString, const bool bShouldRefreshAfterward = true);
 
 	bool DoesTreeViewHaveVisibleChildren() const;
@@ -141,12 +131,8 @@ public:
 	 * Determines the style of the tree (flat list or hierarchy)
 	 */
 	EObjectMixerTreeViewMode GetTreeViewMode();
-	/**
-	 * Determine the style of the tree (flat list or hierarchy)
-	 */
-	void SetTreeViewMode(EObjectMixerTreeViewMode InViewMode);
-
-	void ToggleFilterActive(const FString& FilterName);
+	
+	const TArray<TSharedRef<IObjectMixerEditorListFilter>>& GetShowFilters();
 	void EvaluateIfRowsPassFilters(const bool bShouldRefreshAfterward = true);
 
 	// Sorting
@@ -182,7 +168,7 @@ public:
 		);
 	};
 
-private:
+protected:
 
 	virtual void Tick(const FGeometry& AllottedGeometry, const double InCurrentTime, const float InDeltaTime) override;
 
@@ -196,9 +182,6 @@ private:
 	TSharedPtr<SHeaderRow> HeaderRow;
 	TSharedRef<SWidget> GenerateHeaderRowContextMenu() const;
 
-	TSharedRef<SWidget> GenerateToolbar();
-	TSharedRef<SWidget> OnGenerateAddObjectButtonMenu() const;
-
 	bool bShouldRebuild = false;
 	
 	/**
@@ -219,25 +202,11 @@ private:
 	void AddBuiltinColumnsToHeaderRow();
 	TSharedPtr<SHeaderRow> GenerateHeaderRow();
 	ECheckBoxState HeaderCheckBoxState = ECheckBoxState::Checked;
-
-	void SetupFilters();
-
-	TSharedRef<SWidget> OnGenerateFilterClassMenu();
-	TSharedRef<SWidget> BuildShowOptionsMenu();
 	
 	void FlushMemory(const bool bShouldKeepMemoryAllocated);
 	
 	void SetAllGroupsCollapsed();
 
-	// Search
-	
-	void OnListViewSearchTextChanged(const FText& Text);
-
-	TSharedPtr<SSearchBox> ListSearchBoxPtr;
-	TSharedPtr<SComboButton> ViewOptionsComboButton;
-	
-	TSharedPtr<SBox> ListBoxContainerPtr;
-
 	//  Tree View Implementation
 
 	/** Saves tree item expanded states to be recalled after the tree view is regenerated. */
@@ -253,8 +222,6 @@ private:
 
 	void SetChildExpansionRecursively(const FObjectMixerEditorListRowPtr& InRow, const bool bNewIsExpanded) const;
 
-	TArray<TSharedRef<IObjectMixerEditorListFilter>> ShowFilters;
-
 	TSharedPtr<STreeView<FObjectMixerEditorListRowPtr>> TreeViewPtr;
 
 	TMap<FString, bool> TreeItemExpansionStateCache;
diff --git a/Engine/Source/Editor/ObjectMixer/ObjectMixer/Public/Views/MainPanel/ObjectMixerEditorMainPanel.h b/Engine/Source/Editor/ObjectMixer/ObjectMixer/Public/Views/MainPanel/ObjectMixerEditorMainPanel.h
--- a/Engine/Source/Editor/ObjectMixer/ObjectMixer/Public/Views/MainPanel/ObjectMixerEditorMainPanel.h
+++ b/Engine/Source/Editor/ObjectMixer/ObjectMixer/Public/Views/MainPanel/ObjectMixerEditorMainPanel.h
@@ -2,21 +2,32 @@
 
 #pragma once
 
-#include "Views/List/ObjectMixerEditorList.h"
+#include "ObjectFilter/ObjectMixerEditorObjectFilter.h"
+#include "Views/List/ObjectMixerEditorListFilters/IObjectMixerEditorListFilter.h"
+
+#include "UObject/StrongObjectPtr.h"
 #include "Widgets/SWidget.h"
 
 class FObjectMixerEditorList;
 class SObjectMixerEditorMainPanel;
 
+DECLARE_MULTICAST_DELEGATE(FOnObjectMixerCategoryMapChanged)
+
 class OBJECTMIXEREDITOR_API FObjectMixerEditorMainPanel : public TSharedFromThis<FObjectMixerEditorMainPanel>
 {
 public:
-	FObjectMixerEditorMainPanel();
+	FObjectMixerEditorMainPanel(const FName InModuleName)
+	: ModuleName(InModuleName)
+	{}
 
 	~FObjectMixerEditorMainPanel() = default;
 
+	void Init();
+
 	TSharedRef<SWidget> GetOrCreateWidget();
 
+	void RegenerateListModel();
+
 	/**
 	 * Regenerate the list items and refresh the list. Call when adding or removing items.
 	 */
@@ -28,18 +39,155 @@ public:
 	 */
 	void RefreshList() const;
 
-	TWeakPtr<FObjectMixerEditorList> GetEditorList() const
+	TWeakPtr<FObjectMixerEditorList> GetEditorListModel() const
 	{
-		return EditorList;
+		return EditorListModel;
 	}
 
-	void OnClassSelectionChanged(UClass* InNewClass) const;
+	void RebuildCategorySelector();
+
+	FString GetSearchStringFromSearchInputField() const;
+
+	void OnClassSelectionChanged(UClass* InNewClass);
 	TObjectPtr<UClass> GetClassSelection() const;
 	bool IsClassSelected(UClass* InNewClass) const;
 
+	UObjectMixerObjectFilter* GetObjectFilter();
+
+	void CacheObjectFilterObject();
+
+	/**
+	 * Get the style of the tree (flat list or hierarchy)
+	 */
+	EObjectMixerTreeViewMode GetTreeViewMode()
+	{
+		return TreeViewMode;
+	}
+	/**
+	 * Set the style of the tree (flat list or hierarchy)
+	 */
+	void SetTreeViewMode(EObjectMixerTreeViewMode InViewMode)
+	{
+		TreeViewMode = InViewMode;
+		RequestRebuildList();
+	}
+
+	/**
+	 * Returns result from Filter->GetObjectClassesToFilter.
+	 */
+	TSet<UClass*> GetObjectClassesToFilter()
+	{
+		if (const UObjectMixerObjectFilter* Filter = GetObjectFilter())
+		{
+			return Filter->GetObjectClassesToFilter();
+		}
+		
+		return {};
+	}
+
+	/**
+	 * Returns result from Filter->GetObjectClassesToPlace.
+	 */
+	TSet<TSubclassOf<AActor>> GetObjectClassesToPlace()
+	{
+		TSet<TSubclassOf<AActor>> ReturnValue;
+
+		if (const UObjectMixerObjectFilter* Filter = GetObjectFilter())
+		{
+			ReturnValue = Filter->GetObjectClassesToPlace();
+		}
+		
+		return ReturnValue;
+	}
+
+	const TArray<TSharedRef<IObjectMixerEditorListFilter>>& GetShowFilters() const;
+
+	/**
+	 * Get the row that has solo visibility. All other rows should be set to temporarily invisible in editor.
+	 */
+	TWeakPtr<FObjectMixerEditorListRow> GetSoloRow()
+	{
+		return SoloRow;
+	}
+
+	/**
+	 * Set the row that has solo visibility. This does not set temporary editor invisibility for other rows.
+	 */
+	void SetSoloRow(TSharedRef<FObjectMixerEditorListRow> InRow)
+	{
+		SoloRow = InRow;
+	}
+
+	/**
+	 * Clear the row that has solo visibility. This does not remove temporary editor invisibility for other rows.
+	 */
+	void ClearSoloRow()
+	{
+		SoloRow = nullptr;
+	}
+
+	TSubclassOf<UObjectMixerObjectFilter> GetObjectFilterClass() const
+	{
+		return ObjectFilterClass;
+	}
+
+	void SetObjectFilterClass(UClass* InObjectFilterClass)
+	{
+		if (ensureAlwaysMsgf(InObjectFilterClass->IsChildOf(UObjectMixerObjectFilter::StaticClass()), TEXT("%hs: Class '%s' is not a child of UObjectMixerObjectFilter."), __FUNCTION__, *InObjectFilterClass->GetName()))
+		{
+			ObjectFilterClass = InObjectFilterClass;
+			CacheObjectFilterObject();
+			RequestRebuildList();
+		}
+	}
+
+	FName GetModuleName() const
+	{
+		return ModuleName;
+	}
+
+	// User Categorization
+
+	/**
+	 * Add set of objects to a category in the map, or create a new category if one does not exist,
+	 */
+	void AddObjectsToCategory(const FName& CategoryName, const TSet<FSoftObjectPath>& ObjectsToAdd) const;
+	void RemoveObjectsFromCategory(const FName& CategoryName, const TSet<FSoftObjectPath>& ObjectsToRemove) const;
+	void RemoveCategory(const FName& CategoryName) const;
+	bool IsObjectInCategory(const FName& CategoryName, const FSoftObjectPath& InObject) const;
+	TSet<FName> GetCategoriesForObject(const FSoftObjectPath& InObject) const;
+	TSet<FName> GetAllCategories() const;
+	FOnObjectMixerCategoryMapChanged& GetOnObjectMixerCategoryMapChanged()
+	{
+		return OnObjectMixerCategoryMapChanged;
+	}
+
+	/**
+	 * Returns the categories selected by the user. If the set is empty, consider "All" categories to be selected.
+	 */
+	const TSet<FName>& GetCurrentCategorySelection() const;
+
 private:
 
 	TSharedPtr<SObjectMixerEditorMainPanel> MainPanelWidget;
 
-	TSharedPtr<FObjectMixerEditorList> EditorList;
+	TSharedPtr<FObjectMixerEditorList> EditorListModel;
+
+	TStrongObjectPtr<UObjectMixerObjectFilter> ObjectFilterPtr;
+
+	/**
+	 * The class used to generate property edit columns
+	 */
+	TSubclassOf<UObjectMixerObjectFilter> ObjectFilterClass;
+
+	/**
+	 * Determines the style of the tree (flat list or hierarchy)
+	 */
+	EObjectMixerTreeViewMode TreeViewMode = EObjectMixerTreeViewMode::FolderObjectSubObject;
+
+	TWeakPtr<FObjectMixerEditorListRow> SoloRow = nullptr;
+
+	FName ModuleName = NAME_None;
+
+	FOnObjectMixerCategoryMapChanged OnObjectMixerCategoryMapChanged;
 };
diff --git a/Engine/Source/Editor/ObjectMixer/ObjectMixer/Public/Views/MainPanel/SObjectMixerEditorMainPanel.h b/Engine/Source/Editor/ObjectMixer/ObjectMixer/Public/Views/MainPanel/SObjectMixerEditorMainPanel.h
--- a/Engine/Source/Editor/ObjectMixer/ObjectMixer/Public/Views/MainPanel/SObjectMixerEditorMainPanel.h
+++ b/Engine/Source/Editor/ObjectMixer/ObjectMixer/Public/Views/MainPanel/SObjectMixerEditorMainPanel.h
@@ -2,6 +2,10 @@
 
 #pragma once
 
+#include "ObjectFilter/ObjectMixerEditorObjectFilter.h"
+#include "Views/List/ObjectMixerEditorListFilters/IObjectMixerEditorListFilter.h"
+
+#include "Styling/SlateTypes.h"
 #include "Widgets/SCompoundWidget.h"
 
 class FObjectMixerEditorMainPanel;
@@ -18,10 +22,60 @@ public:
 
 	void Construct(const FArguments& InArgs, const TSharedRef<FObjectMixerEditorMainPanel>& InMainPanel);
 
+	FString GetSearchStringFromSearchInputField() const;
+	void SetSearchStringInSearchInputField(const FString InSearchString) const;
+	void ExecuteListViewSearchOnAllRows(const FString& SearchString, const bool bShouldRefreshAfterward = true);
+
+	/**
+	 * Determines the style of the tree (flat list or hierarchy)
+	 */
+	EObjectMixerTreeViewMode GetTreeViewMode();
+	/**
+	 * Determine the style of the tree (flat list or hierarchy)
+	 */
+	void SetTreeViewMode(EObjectMixerTreeViewMode InViewMode);
+
+	void ToggleFilterActive(const FString& FilterName);
+
+	const TArray<TSharedRef<class IObjectMixerEditorListFilter>>& GetShowFilters()
+	{
+		return ShowFilters;	
+	}
+
+	const TSet<FName>& GetCurrentCategorySelection();
+	
+	void RebuildCategorySelector();
+
 	virtual ~SObjectMixerEditorMainPanel() override;
 
 private:
 
 	/** A reference to the struct that controls this widget */
-	TWeakPtr<FObjectMixerEditorMainPanel> MainPanel;
+	TWeakPtr<FObjectMixerEditorMainPanel> MainPanelModel;
+
+	TArray<TSharedRef<IObjectMixerEditorListFilter>> ShowFilters;
+
+	TSharedPtr<class SSearchBox> SearchBoxPtr;
+	TSharedPtr<class SComboButton> ViewOptionsComboButton;
+
+	TSharedRef<SWidget> GenerateToolbar();
+	TSharedRef<SWidget> OnGenerateAddObjectButtonMenu() const;
+
+	TSharedRef<SWidget> OnGenerateFilterClassMenu();
+	TSharedRef<SWidget> BuildShowOptionsMenu();
+
+	void OnSearchTextChanged(const FText& Text);
+
+	// User Categorization
+	
+	TSharedPtr<class SWrapBox> CategorySelectorBox;
+	TSet<FName> CurrentCategorySelection;
+
+	void ResetCurrentCategorySelection()
+	{
+		CurrentCategorySelection.Reset();
+	}
+	
+	void OnCategoryCheckedChanged(ECheckBoxState State, FName SectionName);
+	ECheckBoxState IsCategoryChecked(FName Section) const;
 };
diff --git a/Engine/Source/Editor/ObjectMixer/ObjectMixer/Public/Views/Widgets/ObjectMixerEditorListMenuContext.h b/Engine/Source/Editor/ObjectMixer/ObjectMixer/Public/Views/Widgets/ObjectMixerEditorListMenuContext.h
new file mode 100644
--- /dev/null
+++ b/Engine/Source/Editor/ObjectMixer/ObjectMixer/Public/Views/Widgets/ObjectMixerEditorListMenuContext.h
@@ -0,0 +1,44 @@
+// Copyright Epic Games, Inc. All Rights Reserved.
+
+#pragma once
+
+#include "Views/List/ObjectMixerEditorListRow.h"
+
+#include "Types/SlateEnums.h"
+
+#include "ObjectMixerEditorListMenuContext.generated.h"	
+
+UCLASS()
+class OBJECTMIXEREDITOR_API UObjectMixerEditorListMenuContext : public UObject
+{
+
+	GENERATED_BODY()
+	
+public:
+
+	struct FObjectMixerEditorListMenuContextData
+	{
+		TArray<FObjectMixerEditorListRowPtr> SelectedItems;
+		TWeakPtr<class FObjectMixerEditorMainPanel> MainPanelPtr;
+	};
+
+	static TSharedPtr<SWidget> CreateContextMenu(const FObjectMixerEditorListMenuContextData InData);
+	static TSharedPtr<SWidget> BuildContextMenu(const FObjectMixerEditorListMenuContextData& InData);
+	static void RegisterContextMenu();
+
+	FObjectMixerEditorListMenuContextData Data;
+	
+	static FName DefaultContextBaseMenuName;
+
+private:
+
+	static void OnTextCommitted(const FText& InText, ETextCommit::Type InCommitType, UObjectMixerEditorListMenuContext* Context);
+	static void OnTextChanged(const FText& InText, UObjectMixerEditorListMenuContext* Context);
+	
+	static void OnClickCategoryMenuEntry(const FName Key, UObjectMixerEditorListMenuContext* Context);
+	static void AddObjectsToCategory(const FName Key, UObjectMixerEditorListMenuContext* Context);
+	static void RemoveObjectsFromCategory(const FName Key, UObjectMixerEditorListMenuContext* Context);
+	static bool AreAllObjectsInCategory(const FName Key, UObjectMixerEditorListMenuContext* Context);
+
+	TSharedPtr<class SEditableTextBox> EditableText;
+};
