Object Mixer

-Implement selection sync toggle button
-Separate Serialized Data into its own class
-Double-click Row to focus in Viewport

diff --git a/Engine/Source/Editor/ObjectMixer/LightMixer/Private/LightMixerModule.cpp b/Engine/Source/Editor/ObjectMixer/LightMixer/Private/LightMixerModule.cpp
--- a/Engine/Source/Editor/ObjectMixer/LightMixer/Private/LightMixerModule.cpp
+++ b/Engine/Source/Editor/ObjectMixer/LightMixer/Private/LightMixerModule.cpp
@@ -33,6 +33,12 @@ FLightMixerModule& FLightMixerModule::Get()
 	return FModuleManager::LoadModuleChecked< FLightMixerModule >("LightMixer");
 }
 
+void FLightMixerModule::OpenProjectSettings()
+{
+	FModuleManager::LoadModuleChecked<ISettingsModule>("Settings")
+		.ShowViewer("Project", "Editor", "Light Mixer");
+}
+
 void FLightMixerModule::Initialize()
 {
 	FObjectMixerEditorModule::Initialize();
@@ -66,7 +72,7 @@ void FLightMixerModule::RegisterProjectSettings() const
 	{
 		// User Project Settings
 		const TSharedPtr<ISettingsSection> ProjectSettingsSectionPtr = SettingsModule->RegisterSettings(
-			"Project", "Plugins", "Light Mixer",
+			"Project", "Editor", "Light Mixer",
 			LOCTEXT("LightMixerSettingsCategoryDisplayName", "Light Mixer"),
 			LOCTEXT("LightMixerSettingsDescription", "Configure Light Mixer user settings"),
 			GetMutableDefault<ULightMixerProjectSettings>());
@@ -77,7 +83,7 @@ void FLightMixerModule::UnregisterProjectSettings() const
 {
 	if (ISettingsModule* SettingsModule = FModuleManager::GetModulePtr<ISettingsModule>("Settings"))
 	{
-		SettingsModule->UnregisterSettings("Project", "Plugins", "Light Mixer");
+		SettingsModule->UnregisterSettings("Project", "Editor", "Light Mixer");
 	}
 }
 
diff --git a/Engine/Source/Editor/ObjectMixer/LightMixer/Public/LightMixerModule.h b/Engine/Source/Editor/ObjectMixer/LightMixer/Public/LightMixerModule.h
--- a/Engine/Source/Editor/ObjectMixer/LightMixer/Public/LightMixerModule.h
+++ b/Engine/Source/Editor/ObjectMixer/LightMixer/Public/LightMixerModule.h
@@ -15,6 +15,8 @@ public:
 
 	static FLightMixerModule& Get();
 	
+	static void OpenProjectSettings();
+	
 	//~ Begin FObjectMixerEditorModule overrides
 	virtual void Initialize() override;
 	virtual FName GetModuleName() override;
diff --git a/Engine/Source/Editor/ObjectMixer/LightMixer/Public/LightMixerProjectSettings.h b/Engine/Source/Editor/ObjectMixer/LightMixer/Public/LightMixerProjectSettings.h
--- a/Engine/Source/Editor/ObjectMixer/LightMixer/Public/LightMixerProjectSettings.h
+++ b/Engine/Source/Editor/ObjectMixer/LightMixer/Public/LightMixerProjectSettings.h
@@ -4,7 +4,7 @@
 
 #include "LightMixerProjectSettings.generated.h"
 
-UCLASS(config = Engine, defaultconfig)
+UCLASS(config = ObjectMixer, defaultconfig)
 class LIGHTMIXER_API ULightMixerProjectSettings : public UObject
 {
 	GENERATED_BODY()
diff --git a/Engine/Source/Editor/ObjectMixer/ObjectMixer/Private/ObjectMixerEditorModule.cpp b/Engine/Source/Editor/ObjectMixer/ObjectMixer/Private/ObjectMixerEditorModule.cpp
--- a/Engine/Source/Editor/ObjectMixer/ObjectMixer/Private/ObjectMixerEditorModule.cpp
+++ b/Engine/Source/Editor/ObjectMixer/ObjectMixer/Private/ObjectMixerEditorModule.cpp
@@ -98,6 +98,12 @@ FObjectMixerEditorModule& FObjectMixerEditorModule::Get()
 	return FModuleManager::LoadModuleChecked< FObjectMixerEditorModule >("ObjectMixerEditor");
 }
 
+void FObjectMixerEditorModule::OpenProjectSettings()
+{
+	FModuleManager::LoadModuleChecked<ISettingsModule>("Settings")
+		.ShowViewer("Project", "Editor", "Object Mixer");
+}
+
 FName FObjectMixerEditorModule::GetModuleName()
 {
 	return "ObjectMixerEditor";
diff --git a/Engine/Source/Editor/ObjectMixer/ObjectMixer/Private/ObjectMixerEditorSerializedData.cpp b/Engine/Source/Editor/ObjectMixer/ObjectMixer/Private/ObjectMixerEditorSerializedData.cpp
new file mode 100644
--- /dev/null
+++ b/Engine/Source/Editor/ObjectMixer/ObjectMixer/Private/ObjectMixerEditorSerializedData.cpp
@@ -0,0 +1,141 @@
+// Copyright Epic Games, Inc. All Rights Reserved.
+
+#include "ObjectMixerEditorSerializedData.h"
+
+#include "Algo/Find.h"
+
+FObjectMixerSerializationData* UObjectMixerEditorSerializedData::FindSerializationDataByFilterClassName(const FName& FilterClassName)
+{
+	if (FObjectMixerSerializationData* Match = Algo::FindByPredicate(SerializedData,
+		[FilterClassName](const FObjectMixerSerializationData& Comparator)
+		{
+			return Comparator.FilterClassName.IsEqual(FilterClassName);
+		}))
+	{
+		return Match;
+	}
+
+	return nullptr;
+}
+
+void UObjectMixerEditorSerializedData::AddObjectsToCategory(const FName& FilterClassName, const FName& CategoryName, const TSet<FSoftObjectPath>& ObjectsToAdd)
+{
+	FObjectMixerSerializationData* Data = FindSerializationDataByFilterClassName(FilterClassName);
+
+	if (!Data)
+	{
+		SerializedData.Add({FilterClassName});
+		Data = FindSerializationDataByFilterClassName(FilterClassName);
+	}
+	
+	if (Data)
+	{
+		if (FObjectMixerCategoryObjectSet* Match = Algo::FindByPredicate(Data->SerializedCategories,
+			[CategoryName](const FObjectMixerCategoryObjectSet& Comparator)
+			{
+				return Comparator.CategoryName.IsEqual(CategoryName);
+			}))
+		{
+			(*Match).CategoryObjects.Append(ObjectsToAdd);
+		}
+		else
+		{
+			Data->SerializedCategories.Add({CategoryName, ObjectsToAdd});
+		}
+
+		SaveConfig();
+	}
+}
+
+void UObjectMixerEditorSerializedData::RemoveObjectsFromCategory(const FName& FilterClassName, const FName& CategoryName, const TSet<FSoftObjectPath>& ObjectsToRemove)
+{
+	if (FObjectMixerSerializationData* Data = FindSerializationDataByFilterClassName(FilterClassName))
+	{
+		if (FObjectMixerCategoryObjectSet* Match = Algo::FindByPredicate(Data->SerializedCategories,
+			[CategoryName](const FObjectMixerCategoryObjectSet& Comparator)
+			{
+				return Comparator.CategoryName.IsEqual(CategoryName);
+			}))
+		{
+			for (const FSoftObjectPath& ObjectPath : ObjectsToRemove)
+			{
+				(*Match).CategoryObjects.Remove(ObjectPath);
+			}
+			
+			SaveConfig();
+		}
+	}
+}
+
+void UObjectMixerEditorSerializedData::RemoveCategory(const FName& FilterClassName, const FName& CategoryName)
+{
+	if (FObjectMixerSerializationData* Data = FindSerializationDataByFilterClassName(FilterClassName))
+	{
+		if (FObjectMixerCategoryObjectSet* Match = Algo::FindByPredicate(Data->SerializedCategories,
+			[CategoryName](const FObjectMixerCategoryObjectSet& Comparator)
+			{
+				return Comparator.CategoryName.IsEqual(CategoryName);
+			}))
+		{
+			(*Match).CategoryObjects.Empty();
+
+			Data->SerializedCategories.Remove(*Match);
+
+			SaveConfig();
+		}
+	}
+}
+
+bool UObjectMixerEditorSerializedData::IsObjectInCategory(const FName& FilterClassName, const FName& CategoryName, const FSoftObjectPath& InObject)
+{
+	if (FObjectMixerSerializationData* Data = FindSerializationDataByFilterClassName(FilterClassName))
+	{
+		if (const FObjectMixerCategoryObjectSet* Match = Algo::FindByPredicate(Data->SerializedCategories,
+			[CategoryName](const FObjectMixerCategoryObjectSet& Comparator)
+			{
+				return Comparator.CategoryName.IsEqual(CategoryName);
+			}))
+		{
+			return (*Match).CategoryObjects.Contains(InObject);
+		}
+	}
+
+	return false;
+}
+
+TSet<FName> UObjectMixerEditorSerializedData::GetCategoriesForObject(const FName& FilterClassName, const FSoftObjectPath& InObject)
+{
+	if (FObjectMixerSerializationData* Data = FindSerializationDataByFilterClassName(FilterClassName))
+	{
+		TSet<FName> CategoriesWithObject;
+
+		for (const FObjectMixerCategoryObjectSet& SerializedCategory : Data->SerializedCategories)
+		{
+			if (IsObjectInCategory(FilterClassName, SerializedCategory.CategoryName, InObject))
+			{
+				CategoriesWithObject.Add(SerializedCategory.CategoryName);
+			}
+		}
+
+		return CategoriesWithObject;
+	}
+
+	return {};
+}
+
+TSet<FName> UObjectMixerEditorSerializedData::GetAllCategories(const FName& FilterClassName)
+{
+	if (FObjectMixerSerializationData* Data = FindSerializationDataByFilterClassName(FilterClassName))
+	{
+		TSet<FName> Categories;
+
+		for (const FObjectMixerCategoryObjectSet& SerializedCategory : Data->SerializedCategories)
+		{
+			Categories.Add(SerializedCategory.CategoryName);
+		}
+
+		return Categories;
+	}
+
+	return {};
+}
diff --git a/Engine/Source/Editor/ObjectMixer/ObjectMixer/Private/Views/List/SObjectMixerEditorList.cpp b/Engine/Source/Editor/ObjectMixer/ObjectMixer/Private/Views/List/SObjectMixerEditorList.cpp
--- a/Engine/Source/Editor/ObjectMixer/ObjectMixer/Private/Views/List/SObjectMixerEditorList.cpp
+++ b/Engine/Source/Editor/ObjectMixer/ObjectMixer/Private/Views/List/SObjectMixerEditorList.cpp
@@ -16,6 +16,7 @@
 #include "Algo/AllOf.h"
 #include "Algo/AnyOf.h"
 #include "Editor.h"
+#include "LevelEditorViewport.h"
 #include "SceneOutlinerMenuContext.h"
 #include "Selection.h"
 #include "Framework/MultiBox/MultiBoxBuilder.h"
@@ -84,7 +85,11 @@ void SObjectMixerEditorList::Construct(const FArguments& InArgs, TSharedRef<FObj
 				})
 				.OnSelectionChanged_Lambda([this] (const FObjectMixerEditorListRowPtr& Row, const ESelectInfo::Type SelectionType)
 				{
-					if (GEditor && FSlateApplication::Get().GetModifierKeys().IsAltDown())
+					const bool bSyncSelectionEnabled = GetDefault<UObjectMixerEditorProjectSettings>()->bSyncSelection;
+					const bool bIsAltDown = FSlateApplication::Get().GetModifierKeys().IsAltDown();
+					const bool bShouldSyncSelection = bSyncSelectionEnabled ? !bIsAltDown : bIsAltDown;
+						
+					if (GEditor && bShouldSyncSelection)
 					{
 						TArray<AActor*> ActorsToSelect;
 						for (const TSharedPtr<FObjectMixerEditorListRow>& SelectedRow : TreeViewPtr->GetSelectedItems())
@@ -117,6 +122,34 @@ void SObjectMixerEditorList::Construct(const FArguments& InArgs, TSharedRef<FObj
 						}
 					}
 				})
+				.OnMouseButtonDoubleClick_Lambda([this] (FObjectMixerEditorListRowPtr Row)
+				{							
+					if (Row->GetRowType() == FObjectMixerEditorListRow::MatchingObject ||
+						Row->GetRowType() == FObjectMixerEditorListRow::ContainerObject)
+					{
+						if (GCurrentLevelEditingViewportClient)
+						{
+							AActor* Actor = Cast<AActor>(Row->GetObject());
+
+							if (!Actor)
+							{
+								Actor = Row->GetObject()->GetTypedOuter<AActor>();
+							}
+
+							if (Actor)
+							{
+								FVector Origin;
+								FVector Extents;
+								Actor->GetActorBounds(false, Origin, Extents, true);
+								GCurrentLevelEditingViewportClient->FocusViewportOnBox(FBox(Origin - Extents, Origin + Extents));
+							}
+						}
+					}
+					else
+					{
+						Row->SetIsTreeViewItemExpanded(!Row->GetIsTreeViewItemExpanded());
+					}
+				})
 				.TreeItemsSource(&VisibleTreeViewObjects)
 				.OnGenerateRow_Lambda([this](FObjectMixerEditorListRowPtr Row, const TSharedRef<STableViewBase>& OwnerTable)
 					{
diff --git a/Engine/Source/Editor/ObjectMixer/ObjectMixer/Private/Views/MainPanel/ObjectMixerEditorMainPanel.cpp b/Engine/Source/Editor/ObjectMixer/ObjectMixer/Private/Views/MainPanel/ObjectMixerEditorMainPanel.cpp
--- a/Engine/Source/Editor/ObjectMixer/ObjectMixer/Private/Views/MainPanel/ObjectMixerEditorMainPanel.cpp
+++ b/Engine/Source/Editor/ObjectMixer/ObjectMixer/Private/Views/MainPanel/ObjectMixerEditorMainPanel.cpp
@@ -2,7 +2,7 @@
 
 #include "Views/MainPanel/ObjectMixerEditorMainPanel.h"
 
-#include "ObjectMixerEditorProjectSettings.h"
+#include "ObjectMixerEditorSerializedData.h"
 #include "Views/List/ObjectMixerEditorList.h"
 #include "Views/MainPanel/SObjectMixerEditorMainPanel.h"
 
@@ -101,7 +101,7 @@ void FObjectMixerEditorMainPanel::AddObjectsToCategory(const FName& CategoryName
 {
 	if (const TSubclassOf<UObjectMixerObjectFilter> Filter = GetObjectFilterClass())
 	{
-		if (UObjectMixerEditorProjectSettings* Settings = GetMutableDefault<UObjectMixerEditorProjectSettings>())
+		if (UObjectMixerEditorSerializedData* Settings = GetMutableDefault<UObjectMixerEditorSerializedData>())
 		{
 			Settings->AddObjectsToCategory(Filter->GetFName(), CategoryName, ObjectsToAdd);
 
@@ -114,7 +114,7 @@ void FObjectMixerEditorMainPanel::RemoveObjectsFromCategory(const FName& Categor
 {
 	if (const TSubclassOf<UObjectMixerObjectFilter> Filter = GetObjectFilterClass())
 	{
-		if (UObjectMixerEditorProjectSettings* Settings = GetMutableDefault<UObjectMixerEditorProjectSettings>())
+		if (UObjectMixerEditorSerializedData* Settings = GetMutableDefault<UObjectMixerEditorSerializedData>())
 		{
 			Settings->RemoveObjectsFromCategory(Filter->GetFName(), CategoryName, ObjectsToRemove);
 
@@ -127,7 +127,7 @@ void FObjectMixerEditorMainPanel::RemoveCategory(const FName& CategoryName) cons
 {
 	if (const TSubclassOf<UObjectMixerObjectFilter> Filter = GetObjectFilterClass())
 	{
-		if (UObjectMixerEditorProjectSettings* Settings = GetMutableDefault<UObjectMixerEditorProjectSettings>())
+		if (UObjectMixerEditorSerializedData* Settings = GetMutableDefault<UObjectMixerEditorSerializedData>())
 		{
 			Settings->RemoveCategory(Filter->GetFName(), CategoryName);
 
@@ -140,7 +140,7 @@ bool FObjectMixerEditorMainPanel::IsObjectInCategory(const FName& CategoryName,
 {
 	if (const TSubclassOf<UObjectMixerObjectFilter> Filter = GetObjectFilterClass())
 	{
-		if (UObjectMixerEditorProjectSettings* Settings = GetMutableDefault<UObjectMixerEditorProjectSettings>())
+		if (UObjectMixerEditorSerializedData* Settings = GetMutableDefault<UObjectMixerEditorSerializedData>())
 		{
 			return Settings->IsObjectInCategory(Filter->GetFName(), CategoryName, InObject);
 		}
@@ -153,7 +153,7 @@ TSet<FName> FObjectMixerEditorMainPanel::GetCategoriesForObject(const FSoftObjec
 {
 	if (const TSubclassOf<UObjectMixerObjectFilter> Filter = GetObjectFilterClass())
 	{
-		if (UObjectMixerEditorProjectSettings* Settings = GetMutableDefault<UObjectMixerEditorProjectSettings>())
+		if (UObjectMixerEditorSerializedData* Settings = GetMutableDefault<UObjectMixerEditorSerializedData>())
 		{
 			return Settings->GetCategoriesForObject(Filter->GetFName(), InObject);
 		}
@@ -166,7 +166,7 @@ TSet<FName> FObjectMixerEditorMainPanel::GetAllCategories() const
 {
 	if (const TSubclassOf<UObjectMixerObjectFilter> Filter = GetObjectFilterClass())
 	{
-		if (UObjectMixerEditorProjectSettings* Settings = GetMutableDefault<UObjectMixerEditorProjectSettings>())
+		if (UObjectMixerEditorSerializedData* Settings = GetMutableDefault<UObjectMixerEditorSerializedData>())
 		{
 			return Settings->GetAllCategories(Filter->GetFName());
 		}
diff --git a/Engine/Source/Editor/ObjectMixer/ObjectMixer/Private/Views/MainPanel/SObjectMixerEditorMainPanel.cpp b/Engine/Source/Editor/ObjectMixer/ObjectMixer/Private/Views/MainPanel/SObjectMixerEditorMainPanel.cpp
--- a/Engine/Source/Editor/ObjectMixer/ObjectMixer/Private/Views/MainPanel/SObjectMixerEditorMainPanel.cpp
+++ b/Engine/Source/Editor/ObjectMixer/ObjectMixer/Private/Views/MainPanel/SObjectMixerEditorMainPanel.cpp
@@ -4,6 +4,7 @@
 
 #include "ObjectMixerEditorLog.h"
 #include "ObjectMixerEditorModule.h"
+#include "ObjectMixerEditorProjectSettings.h"
 #include "Views/List/ObjectMixerEditorList.h"
 #include "Views/MainPanel/ObjectMixerEditorMainPanel.h"
 #include "Views/Widgets/SObjectMixerPlacementAssetMenuEntry.h"
@@ -92,6 +93,41 @@ TSharedRef<SWidget> SObjectMixerEditorMainPanel::GenerateToolbar()
 		.OnTextChanged_Raw(this, &SObjectMixerEditorMainPanel::OnSearchTextChanged)
 	];
 
+	// Selection Sync Toggle
+	ToolbarBox->AddSlot()
+	.HAlign(HAlign_Right)
+	.VAlign(VAlign_Center)
+	.AutoWidth()
+	.Padding(8.f, 1.f, 0.f, 1.f)
+	[
+		SNew(SCheckBox )
+		.Padding(4.f)
+		.ToolTipText(LOCTEXT("SyncSelectionButton_Tooltip", "Sync Selection\nIf enabled, clicking an item in the mixer list will also select the item in the Scene Outliner.\nAlt + Click to select items in mixer without selecting the item in the Scene outliner.\nIf disabled, selections will not sync unless Alt is held. Effectively, this is the opposite behavior."))
+		.Style(&FAppStyle::Get().GetWidgetStyle<FCheckBoxStyle>("ToggleButtonCheckbox"))
+		.ForegroundColor(FSlateColor::UseForeground())
+		.IsChecked_Lambda([]()
+		{
+			if (const UObjectMixerEditorProjectSettings* Settings = GetDefault<UObjectMixerEditorProjectSettings>())
+			{
+				return Settings->bSyncSelection ? ECheckBoxState::Checked : ECheckBoxState::Unchecked;
+			}
+			
+			return ECheckBoxState::Undetermined; 
+		})
+		.OnCheckStateChanged_Lambda([](ECheckBoxState InNewState)
+		{
+			if (UObjectMixerEditorProjectSettings* Settings = GetMutableDefault<UObjectMixerEditorProjectSettings>())
+			{
+				Settings->bSyncSelection = InNewState == ECheckBoxState::Checked ? true : false;
+			}
+		})
+		[
+			SNew(SImage)
+			.ColorAndOpacity(FSlateColor::UseForeground())
+			.Image( FAppStyle::Get().GetBrush("FoliageEditMode.SelectAll") )
+		]
+	];
+
 	// Show Options
 	ToolbarBox->AddSlot()
 	.HAlign(HAlign_Right)
diff --git a/Engine/Source/Editor/ObjectMixer/ObjectMixer/Public/ObjectMixerEditorModule.h b/Engine/Source/Editor/ObjectMixer/ObjectMixer/Public/ObjectMixerEditorModule.h
--- a/Engine/Source/Editor/ObjectMixer/ObjectMixer/Public/ObjectMixerEditorModule.h
+++ b/Engine/Source/Editor/ObjectMixer/ObjectMixer/Public/ObjectMixerEditorModule.h
@@ -24,6 +24,8 @@ public:
 
 	static FObjectMixerEditorModule& Get();
 
+	static void OpenProjectSettings();
+
 	virtual FName GetModuleName();
 	
 	virtual TSharedPtr<SWidget> MakeObjectMixerDialog() const;
diff --git a/Engine/Source/Editor/ObjectMixer/ObjectMixer/Public/ObjectMixerEditorProjectSettings.h b/Engine/Source/Editor/ObjectMixer/ObjectMixer/Public/ObjectMixerEditorProjectSettings.h
--- a/Engine/Source/Editor/ObjectMixer/ObjectMixer/Public/ObjectMixerEditorProjectSettings.h
+++ b/Engine/Source/Editor/ObjectMixer/ObjectMixer/Public/ObjectMixerEditorProjectSettings.h
@@ -2,55 +2,9 @@
 
 #pragma once
 
-#include "Algo/Find.h"
-
 #include "ObjectMixerEditorProjectSettings.generated.h"
 
-USTRUCT()
-struct FObjectMixerCategoryObjectSet
-{
-	GENERATED_BODY()
-	
-	UPROPERTY()
-	FName CategoryName = NAME_None;
-
-	UPROPERTY()
-	TSet<FSoftObjectPath> CategoryObjects = {};
-
-	bool operator==(const FObjectMixerCategoryObjectSet& Other) const
-	{
-		return CategoryName.IsEqual(Other.CategoryName);
-	}
-
-	friend uint32 GetTypeHash (const FObjectMixerCategoryObjectSet& Other)
-	{
-		return GetTypeHash(Other.CategoryName);
-	}
-};
-
-USTRUCT()
-struct FObjectMixerSerializationData
-{
-	GENERATED_BODY()
-	
-	UPROPERTY()
-	FName FilterClassName = NAME_None;
-
-	UPROPERTY()
-	TSet<FObjectMixerCategoryObjectSet> SerializedCategories = {};
-
-	bool operator==(const FObjectMixerSerializationData& Other) const
-	{
-		return FilterClassName.IsEqual(Other.FilterClassName);
-	}
-
-	friend uint32 GetTypeHash (const FObjectMixerSerializationData& Other)
-	{
-		return GetTypeHash(Other.FilterClassName);
-	}
-};
-
-UCLASS(config = Engine, defaultconfig)
+UCLASS(config = ObjectMixer, defaultconfig)
 class OBJECTMIXEREDITOR_API UObjectMixerEditorProjectSettings : public UObject
 {
 	GENERATED_BODY()
@@ -58,149 +12,18 @@ public:
 	
 	UObjectMixerEditorProjectSettings(const FObjectInitializer& ObjectInitializer)
 	{}
+	
+	/**
+	 * If enabled, clicking an item in the mixer list will also select the item in the Scene Outliner.
+	 * Alt + Click to select items in mixer without selecting the item in the Scene outliner.
+	 * If disabled, selections will not sync unless Alt is held. Effectively, this is the opposite behavior.
+	 */
+	UPROPERTY(Config, EditAnywhere, Category="Object Mixer")
+	bool bSyncSelection = true;
 
 	/**
 	 * If false, a new object will be created every time the filter object is accessed.
 	 */
 	UPROPERTY(Config, EditAnywhere, Category="Object Mixer")
 	bool bExpandTreeViewItemsByDefault = true;
-
-	UPROPERTY(Config)
-	TSet<FObjectMixerSerializationData> SerializedData;
-
-	FObjectMixerSerializationData* FindSerializationDataByFilterClassName(const FName& FilterClassName)
-	{
-		if (FObjectMixerSerializationData* Match = Algo::FindByPredicate(SerializedData,
-			[FilterClassName](const FObjectMixerSerializationData& Comparator)
-			{
-				return Comparator.FilterClassName.IsEqual(FilterClassName);
-			}))
-		{
-			return Match;
-		}
-
-		return nullptr;
-	}
-
-	void AddObjectsToCategory(const FName& FilterClassName, const FName& CategoryName, const TSet<FSoftObjectPath>& ObjectsToAdd)
-	{
-		FObjectMixerSerializationData* Data = FindSerializationDataByFilterClassName(FilterClassName);
-
-		if (!Data)
-		{
-			SerializedData.Add({FilterClassName});
-			Data = FindSerializationDataByFilterClassName(FilterClassName);
-		}
-		
-		if (Data)
-		{
-			if (FObjectMixerCategoryObjectSet* Match = Algo::FindByPredicate(Data->SerializedCategories,
-				[CategoryName](const FObjectMixerCategoryObjectSet& Comparator)
-				{
-					return Comparator.CategoryName.IsEqual(CategoryName);
-				}))
-			{
-				(*Match).CategoryObjects.Append(ObjectsToAdd);
-			}
-			else
-			{
-				Data->SerializedCategories.Add({CategoryName, ObjectsToAdd});
-			}
-
-			SaveConfig();
-		}
-	}
-
-	void RemoveObjectsFromCategory(const FName& FilterClassName, const FName& CategoryName, const TSet<FSoftObjectPath>& ObjectsToRemove)
-	{
-		if (FObjectMixerSerializationData* Data = FindSerializationDataByFilterClassName(FilterClassName))
-		{
-			if (FObjectMixerCategoryObjectSet* Match = Algo::FindByPredicate(Data->SerializedCategories,
-				[CategoryName](const FObjectMixerCategoryObjectSet& Comparator)
-				{
-					return Comparator.CategoryName.IsEqual(CategoryName);
-				}))
-			{
-				for (const FSoftObjectPath& ObjectPath : ObjectsToRemove)
-				{
-					(*Match).CategoryObjects.Remove(ObjectPath);
-				}
-				
-				SaveConfig();
-			}
-		}
-	}
-
-	void RemoveCategory(const FName& FilterClassName, const FName& CategoryName)
-	{
-		if (FObjectMixerSerializationData* Data = FindSerializationDataByFilterClassName(FilterClassName))
-		{
-			if (FObjectMixerCategoryObjectSet* Match = Algo::FindByPredicate(Data->SerializedCategories,
-				[CategoryName](const FObjectMixerCategoryObjectSet& Comparator)
-				{
-					return Comparator.CategoryName.IsEqual(CategoryName);
-				}))
-			{
-				(*Match).CategoryObjects.Empty();
-
-				Data->SerializedCategories.Remove(*Match);
-
-				SaveConfig();
-			}
-		}
-	}
-
-	bool IsObjectInCategory(const FName& FilterClassName, const FName& CategoryName, const FSoftObjectPath& InObject)
-	{
-		if (FObjectMixerSerializationData* Data = FindSerializationDataByFilterClassName(FilterClassName))
-		{
-			if (const FObjectMixerCategoryObjectSet* Match = Algo::FindByPredicate(Data->SerializedCategories,
-				[CategoryName](const FObjectMixerCategoryObjectSet& Comparator)
-				{
-					return Comparator.CategoryName.IsEqual(CategoryName);
-				}))
-			{
-				return (*Match).CategoryObjects.Contains(InObject);
-			}
-		}
-
-		return false;
-	}
-
-	TSet<FName> GetCategoriesForObject(const FName& FilterClassName, const FSoftObjectPath& InObject)
-	{
-		if (FObjectMixerSerializationData* Data = FindSerializationDataByFilterClassName(FilterClassName))
-		{
-			TSet<FName> CategoriesWithObject;
-
-			for (const FObjectMixerCategoryObjectSet& SerializedCategory : Data->SerializedCategories)
-			{
-				if (IsObjectInCategory(FilterClassName, SerializedCategory.CategoryName, InObject))
-				{
-					CategoriesWithObject.Add(SerializedCategory.CategoryName);
-				}
-			}
-
-			return CategoriesWithObject;
-		}
-
-		return {};
-	}
-
-	TSet<FName> GetAllCategories(const FName& FilterClassName)
-	{
-		if (FObjectMixerSerializationData* Data = FindSerializationDataByFilterClassName(FilterClassName))
-		{
-			TSet<FName> Categories;
-
-			for (const FObjectMixerCategoryObjectSet& SerializedCategory : Data->SerializedCategories)
-			{
-				Categories.Add(SerializedCategory.CategoryName);
-			}
-
-			return Categories;
-		}
-
-		return {};
-	}
-};
+};
\ No newline at end of file
diff --git a/Engine/Source/Editor/ObjectMixer/ObjectMixer/Public/ObjectMixerEditorSerializedData.h b/Engine/Source/Editor/ObjectMixer/ObjectMixer/Public/ObjectMixerEditorSerializedData.h
new file mode 100644
--- /dev/null
+++ b/Engine/Source/Editor/ObjectMixer/ObjectMixer/Public/ObjectMixerEditorSerializedData.h
@@ -0,0 +1,76 @@
+// Copyright Epic Games, Inc. All Rights Reserved.
+
+#pragma once
+
+#include "ObjectMixerEditorSerializedData.generated.h"
+
+USTRUCT()
+struct FObjectMixerCategoryObjectSet
+{
+	GENERATED_BODY()
+	
+	UPROPERTY()
+	FName CategoryName = NAME_None;
+
+	UPROPERTY()
+	TSet<FSoftObjectPath> CategoryObjects = {};
+
+	bool operator==(const FObjectMixerCategoryObjectSet& Other) const
+	{
+		return CategoryName.IsEqual(Other.CategoryName);
+	}
+
+	friend uint32 GetTypeHash (const FObjectMixerCategoryObjectSet& Other)
+	{
+		return GetTypeHash(Other.CategoryName);
+	}
+};
+
+USTRUCT()
+struct FObjectMixerSerializationData
+{
+	GENERATED_BODY()
+	
+	UPROPERTY()
+	FName FilterClassName = NAME_None;
+
+	UPROPERTY()
+	TSet<FObjectMixerCategoryObjectSet> SerializedCategories = {};
+
+	bool operator==(const FObjectMixerSerializationData& Other) const
+	{
+		return FilterClassName.IsEqual(Other.FilterClassName);
+	}
+
+	friend uint32 GetTypeHash (const FObjectMixerSerializationData& Other)
+	{
+		return GetTypeHash(Other.FilterClassName);
+	}
+};
+
+UCLASS(config = ObjectMixerSerializedData)
+class OBJECTMIXEREDITOR_API UObjectMixerEditorSerializedData : public UObject
+{
+	GENERATED_BODY()
+public:
+	
+	UObjectMixerEditorSerializedData(const FObjectInitializer& ObjectInitializer)
+	{}
+	
+	UPROPERTY(Config)
+	TSet<FObjectMixerSerializationData> SerializedData;
+
+	FObjectMixerSerializationData* FindSerializationDataByFilterClassName(const FName& FilterClassName);
+
+	void AddObjectsToCategory(const FName& FilterClassName, const FName& CategoryName, const TSet<FSoftObjectPath>& ObjectsToAdd);
+
+	void RemoveObjectsFromCategory(const FName& FilterClassName, const FName& CategoryName, const TSet<FSoftObjectPath>& ObjectsToRemove);
+
+	void RemoveCategory(const FName& FilterClassName, const FName& CategoryName);
+
+	bool IsObjectInCategory(const FName& FilterClassName, const FName& CategoryName, const FSoftObjectPath& InObject);
+
+	TSet<FName> GetCategoriesForObject(const FName& FilterClassName, const FSoftObjectPath& InObject);
+
+	TSet<FName> GetAllCategories(const FName& FilterClassName);
+};
