Object Mixer

-Lay groundwork for full context menu

diff --git a/Engine/Source/Editor/ObjectMixer/ObjectMixer/ObjectMixerEditor.build.cs b/Engine/Source/Editor/ObjectMixer/ObjectMixer/ObjectMixerEditor.build.cs
--- a/Engine/Source/Editor/ObjectMixer/ObjectMixer/ObjectMixerEditor.build.cs
+++ b/Engine/Source/Editor/ObjectMixer/ObjectMixer/ObjectMixerEditor.build.cs
@@ -39,6 +39,7 @@ public class ObjectMixerEditor : ModuleRules
 				"SlateCore",
 				"ToolMenus", 
 				"ToolWidgets",
+				"TypedElementFramework",
 				"UnrealEd",
 				"WorkspaceMenuStructure"
 			}
diff --git a/Engine/Source/Editor/ObjectMixer/ObjectMixer/Private/Views/List/SObjectMixerEditorList.cpp b/Engine/Source/Editor/ObjectMixer/ObjectMixer/Private/Views/List/SObjectMixerEditorList.cpp
--- a/Engine/Source/Editor/ObjectMixer/ObjectMixer/Private/Views/List/SObjectMixerEditorList.cpp
+++ b/Engine/Source/Editor/ObjectMixer/ObjectMixer/Private/Views/List/SObjectMixerEditorList.cpp
@@ -7,7 +7,6 @@
 
 #include "ObjectFilter/ObjectMixerEditorObjectFilter.h"
 #include "ObjectMixerEditorLog.h"
-#include "ObjectMixerEditorModule.h"
 #include "ObjectMixerEditorProjectSettings.h"
 #include "ObjectMixerEditorStyle.h"
 #include "Views/List/SObjectMixerEditorListRow.h"
@@ -16,8 +15,9 @@
 #include "Algo/AllOf.h"
 #include "Algo/AnyOf.h"
 #include "Editor.h"
+#include "LevelEditor.h"
+#include "LevelEditorMenuContext.h"
 #include "LevelEditorViewport.h"
-#include "SceneOutlinerMenuContext.h"
 #include "Selection.h"
 #include "Framework/MultiBox/MultiBoxBuilder.h"
 #include "DragAndDrop/AssetDragDropOp.h"
@@ -77,6 +77,7 @@ void SObjectMixerEditorList::Construct(const FArguments& InArgs, TSharedRef<FObj
 					
 					if (GetListModelPtr().IsValid())
 					{
+						SelectedTreeItemsToSelectedInLevelEditor();
 						Widget = UObjectMixerEditorListMenuContext::CreateContextMenu(
 							{GetSelectedTreeViewItems(), GetListModelPtr().Pin()->GetMainPanelModel()});
 					}
@@ -89,37 +90,9 @@ void SObjectMixerEditorList::Construct(const FArguments& InArgs, TSharedRef<FObj
 					const bool bIsAltDown = FSlateApplication::Get().GetModifierKeys().IsAltDown();
 					const bool bShouldSyncSelection = bSyncSelectionEnabled ? !bIsAltDown : bIsAltDown;
 						
-					if (GEditor && bShouldSyncSelection)
+					if (bShouldSyncSelection)
 					{
-						TArray<AActor*> ActorsToSelect;
-						for (const TSharedPtr<FObjectMixerEditorListRow>& SelectedRow : TreeViewPtr->GetSelectedItems())
-						{
-							if (SelectedRow->GetRowType() == FObjectMixerEditorListRow::MatchingObject ||
-								SelectedRow->GetRowType() == FObjectMixerEditorListRow::ContainerObject)
-							{
-								AActor* Actor = Cast<AActor>(SelectedRow->GetObject());
-
-								if (!Actor)
-								{
-									Actor = SelectedRow->GetObject()->GetTypedOuter<AActor>();
-								}
-
-								if (Actor)
-								{
-									ActorsToSelect.Add(Actor);
-								}
-							}
-						}
-
-						if (ActorsToSelect.Num())
-						{
-							GEditor->SelectNone(true, true, true);
-
-							for (AActor* Actor : ActorsToSelect)
-							{
-								GEditor->SelectActor( Actor, true, true, true );
-							}
-						}
+						SelectedTreeItemsToSelectedInLevelEditor();
 					}
 				})
 				.OnMouseButtonDoubleClick_Lambda([this] (FObjectMixerEditorListRowPtr Row)
@@ -524,7 +497,7 @@ TSharedRef<SWidget> SObjectMixerEditorList::GenerateHeaderRowContextMenu() const
 		);
 	}
 
-	return  MenuBuilder.MakeWidget();
+	return MenuBuilder.MakeWidget();
 }
 
 bool SObjectMixerEditorList::AddUniquePropertyColumnsToHeaderRow(
@@ -1167,6 +1140,42 @@ void SObjectMixerEditorList::FindVisibleObjectsAndRequestTreeRefresh()
 	TreeViewPtr->RequestTreeRefresh();
 }
 
+void SObjectMixerEditorList::SelectedTreeItemsToSelectedInLevelEditor() const
+{
+	if (GEditor)
+	{
+		TArray<AActor*> ActorsToSelect;
+		for (const TSharedPtr<FObjectMixerEditorListRow>& SelectedRow : TreeViewPtr->GetSelectedItems())
+		{
+			if (SelectedRow->GetRowType() == FObjectMixerEditorListRow::MatchingObject ||
+				SelectedRow->GetRowType() == FObjectMixerEditorListRow::ContainerObject)
+			{
+				AActor* Actor = Cast<AActor>(SelectedRow->GetObject());
+
+				if (!Actor)
+				{
+					Actor = SelectedRow->GetObject()->GetTypedOuter<AActor>();
+				}
+
+				if (Actor)
+				{
+					ActorsToSelect.Add(Actor);
+				}
+			}
+		}
+
+		if (ActorsToSelect.Num())
+		{
+			GEditor->SelectNone(true, true, true);
+
+			for (AActor* Actor : ActorsToSelect)
+			{
+				GEditor->SelectActor( Actor, true, true, true );
+			}
+		}
+	}
+}
+
 void SObjectMixerEditorList::OnGetRowChildren(FObjectMixerEditorListRowPtr Row, TArray<FObjectMixerEditorListRowPtr>& OutChildren) const
 {
 	if (Row.IsValid())
diff --git a/Engine/Source/Editor/ObjectMixer/ObjectMixer/Public/Views/List/SObjectMixerEditorList.h b/Engine/Source/Editor/ObjectMixer/ObjectMixer/Public/Views/List/SObjectMixerEditorList.h
--- a/Engine/Source/Editor/ObjectMixer/ObjectMixer/Public/Views/List/SObjectMixerEditorList.h
+++ b/Engine/Source/Editor/ObjectMixer/ObjectMixer/Public/Views/List/SObjectMixerEditorList.h
@@ -216,6 +216,8 @@ protected:
 	void GenerateTreeView();
 	void FindVisibleTreeViewObjects();
 	void FindVisibleObjectsAndRequestTreeRefresh();
+
+	void SelectedTreeItemsToSelectedInLevelEditor() const;
 	
 	void OnGetRowChildren(FObjectMixerEditorListRowPtr Row, TArray<FObjectMixerEditorListRowPtr>& OutChildren) const;
 	void OnRowChildExpansionChange(FObjectMixerEditorListRowPtr Row, const bool bIsExpanded, const bool bIsRecursive = false) const;
