Object Mixer

-Rename instances of "Category" to "Collection" with respect to user collections

diff --git a/Engine/Source/Editor/ObjectMixer/LightMixer/Private/LightMixerModule.cpp b/Engine/Source/Editor/ObjectMixer/LightMixer/Private/LightMixerModule.cpp
--- a/Engine/Source/Editor/ObjectMixer/LightMixer/Private/LightMixerModule.cpp
+++ b/Engine/Source/Editor/ObjectMixer/LightMixer/Private/LightMixerModule.cpp
@@ -73,7 +73,7 @@ void FLightMixerModule::RegisterProjectSettings() const
 		// User Project Settings
 		const TSharedPtr<ISettingsSection> ProjectSettingsSectionPtr = SettingsModule->RegisterSettings(
 			"Project", "Editor", "Light Mixer",
-			LOCTEXT("LightMixerSettingsCategoryDisplayName", "Light Mixer"),
+			LOCTEXT("LightMixerSettingsDisplayName", "Light Mixer"),
 			LOCTEXT("LightMixerSettingsDescription", "Configure Light Mixer user settings"),
 			GetMutableDefault<ULightMixerProjectSettings>());
 	}
diff --git a/Engine/Source/Editor/ObjectMixer/ObjectMixer/Private/ObjectMixerEditorModule.cpp b/Engine/Source/Editor/ObjectMixer/ObjectMixer/Private/ObjectMixerEditorModule.cpp
--- a/Engine/Source/Editor/ObjectMixer/ObjectMixer/Private/ObjectMixerEditorModule.cpp
+++ b/Engine/Source/Editor/ObjectMixer/ObjectMixer/Private/ObjectMixerEditorModule.cpp
@@ -218,7 +218,7 @@ void FObjectMixerEditorModule::RegisterProjectSettings() const
 		// User Project Settings
 		const TSharedPtr<ISettingsSection> ProjectSettingsSectionPtr = SettingsModule->RegisterSettings(
 			"Project", "Editor", "Object Mixer",
-			LOCTEXT("ObjectMixerSettingsCategoryDisplayName", "Object Mixer"),
+			LOCTEXT("ObjectMixerSettingsDisplayName", "Object Mixer"),
 			LOCTEXT("ObjectMixerSettingsDescription", "Configure Object Mixer user settings"),
 			GetMutableDefault<UObjectMixerEditorProjectSettings>());	
 	}
diff --git a/Engine/Source/Editor/ObjectMixer/ObjectMixer/Private/ObjectMixerEditorSerializedData.cpp b/Engine/Source/Editor/ObjectMixer/ObjectMixer/Private/ObjectMixerEditorSerializedData.cpp
--- a/Engine/Source/Editor/ObjectMixer/ObjectMixer/Private/ObjectMixerEditorSerializedData.cpp
+++ b/Engine/Source/Editor/ObjectMixer/ObjectMixer/Private/ObjectMixerEditorSerializedData.cpp
@@ -18,7 +18,7 @@ FObjectMixerSerializationData* UObjectMixerEditorSerializedData::FindSerializati
 	return nullptr;
 }
 
-void UObjectMixerEditorSerializedData::AddObjectsToCategory(const FName& FilterClassName, const FName& CategoryName, const TSet<FSoftObjectPath>& ObjectsToAdd)
+void UObjectMixerEditorSerializedData::AddObjectsToCollection(const FName& FilterClassName, const FName& CollectionName, const TSet<FSoftObjectPath>& ObjectsToAdd)
 {
 	FObjectMixerSerializationData* Data = FindSerializationDataByFilterClassName(FilterClassName);
 
@@ -30,36 +30,36 @@ void UObjectMixerEditorSerializedData::AddObjectsToCategory(const FName& FilterC
 	
 	if (Data)
 	{
-		if (FObjectMixerCategoryObjectSet* Match = Algo::FindByPredicate(Data->SerializedCategories,
-			[CategoryName](const FObjectMixerCategoryObjectSet& Comparator)
+		if (FObjectMixerCollectionObjectSet* Match = Algo::FindByPredicate(Data->SerializedCollection,
+			[CollectionName](const FObjectMixerCollectionObjectSet& Comparator)
 			{
-				return Comparator.CategoryName.IsEqual(CategoryName);
+				return Comparator.CollectionName.IsEqual(CollectionName);
 			}))
 		{
-			(*Match).CategoryObjects.Append(ObjectsToAdd);
+			(*Match).CollectionObjects.Append(ObjectsToAdd);
 		}
 		else
 		{
-			Data->SerializedCategories.Add({CategoryName, ObjectsToAdd});
+			Data->SerializedCollection.Add({CollectionName, ObjectsToAdd});
 		}
 
 		SaveConfig();
 	}
 }
 
-void UObjectMixerEditorSerializedData::RemoveObjectsFromCategory(const FName& FilterClassName, const FName& CategoryName, const TSet<FSoftObjectPath>& ObjectsToRemove)
+void UObjectMixerEditorSerializedData::RemoveObjectsFromCollection(const FName& FilterClassName, const FName& CollectionName, const TSet<FSoftObjectPath>& ObjectsToRemove)
 {
 	if (FObjectMixerSerializationData* Data = FindSerializationDataByFilterClassName(FilterClassName))
 	{
-		if (FObjectMixerCategoryObjectSet* Match = Algo::FindByPredicate(Data->SerializedCategories,
-			[CategoryName](const FObjectMixerCategoryObjectSet& Comparator)
+		if (FObjectMixerCollectionObjectSet* Match = Algo::FindByPredicate(Data->SerializedCollection,
+			[CollectionName](const FObjectMixerCollectionObjectSet& Comparator)
 			{
-				return Comparator.CategoryName.IsEqual(CategoryName);
+				return Comparator.CollectionName.IsEqual(CollectionName);
 			}))
 		{
 			for (const FSoftObjectPath& ObjectPath : ObjectsToRemove)
 			{
-				(*Match).CategoryObjects.Remove(ObjectPath);
+				(*Match).CollectionObjects.Remove(ObjectPath);
 			}
 			
 			SaveConfig();
@@ -67,74 +67,74 @@ void UObjectMixerEditorSerializedData::RemoveObjectsFromCategory(const FName& Fi
 	}
 }
 
-void UObjectMixerEditorSerializedData::RemoveCategory(const FName& FilterClassName, const FName& CategoryName)
+void UObjectMixerEditorSerializedData::RemoveCollection(const FName& FilterClassName, const FName& CollectionName)
 {
 	if (FObjectMixerSerializationData* Data = FindSerializationDataByFilterClassName(FilterClassName))
 	{
-		if (FObjectMixerCategoryObjectSet* Match = Algo::FindByPredicate(Data->SerializedCategories,
-			[CategoryName](const FObjectMixerCategoryObjectSet& Comparator)
+		if (FObjectMixerCollectionObjectSet* Match = Algo::FindByPredicate(Data->SerializedCollection,
+			[CollectionName](const FObjectMixerCollectionObjectSet& Comparator)
 			{
-				return Comparator.CategoryName.IsEqual(CategoryName);
+				return Comparator.CollectionName.IsEqual(CollectionName);
 			}))
 		{
-			(*Match).CategoryObjects.Empty();
+			(*Match).CollectionObjects.Empty();
 
-			Data->SerializedCategories.Remove(*Match);
+			Data->SerializedCollection.Remove(*Match);
 
 			SaveConfig();
 		}
 	}
 }
 
-bool UObjectMixerEditorSerializedData::IsObjectInCategory(const FName& FilterClassName, const FName& CategoryName, const FSoftObjectPath& InObject)
+bool UObjectMixerEditorSerializedData::IsObjectInCollection(const FName& FilterClassName, const FName& CollectionName, const FSoftObjectPath& InObject)
 {
 	if (FObjectMixerSerializationData* Data = FindSerializationDataByFilterClassName(FilterClassName))
 	{
-		if (const FObjectMixerCategoryObjectSet* Match = Algo::FindByPredicate(Data->SerializedCategories,
-			[CategoryName](const FObjectMixerCategoryObjectSet& Comparator)
+		if (const FObjectMixerCollectionObjectSet* Match = Algo::FindByPredicate(Data->SerializedCollection,
+			[CollectionName](const FObjectMixerCollectionObjectSet& Comparator)
 			{
-				return Comparator.CategoryName.IsEqual(CategoryName);
+				return Comparator.CollectionName.IsEqual(CollectionName);
 			}))
 		{
-			return (*Match).CategoryObjects.Contains(InObject);
+			return (*Match).CollectionObjects.Contains(InObject);
 		}
 	}
 
 	return false;
 }
 
-TSet<FName> UObjectMixerEditorSerializedData::GetCategoriesForObject(const FName& FilterClassName, const FSoftObjectPath& InObject)
+TSet<FName> UObjectMixerEditorSerializedData::GetCollectionsForObject(const FName& FilterClassName, const FSoftObjectPath& InObject)
 {
 	if (FObjectMixerSerializationData* Data = FindSerializationDataByFilterClassName(FilterClassName))
 	{
-		TSet<FName> CategoriesWithObject;
+		TSet<FName> CollectionsWithObject;
 
-		for (const FObjectMixerCategoryObjectSet& SerializedCategory : Data->SerializedCategories)
+		for (const FObjectMixerCollectionObjectSet& CollectionObjectSet : Data->SerializedCollection)
 		{
-			if (IsObjectInCategory(FilterClassName, SerializedCategory.CategoryName, InObject))
+			if (IsObjectInCollection(FilterClassName, CollectionObjectSet.CollectionName, InObject))
 			{
-				CategoriesWithObject.Add(SerializedCategory.CategoryName);
+				CollectionsWithObject.Add(CollectionObjectSet.CollectionName);
 			}
 		}
 
-		return CategoriesWithObject;
+		return CollectionsWithObject;
 	}
 
 	return {};
 }
 
-TSet<FName> UObjectMixerEditorSerializedData::GetAllCategories(const FName& FilterClassName)
+TSet<FName> UObjectMixerEditorSerializedData::GetAllCollections(const FName& FilterClassName)
 {
 	if (FObjectMixerSerializationData* Data = FindSerializationDataByFilterClassName(FilterClassName))
 	{
-		TSet<FName> Categories;
+		TSet<FName> Collections;
 
-		for (const FObjectMixerCategoryObjectSet& SerializedCategory : Data->SerializedCategories)
+		for (const FObjectMixerCollectionObjectSet& CollectionObjectSet : Data->SerializedCollection)
 		{
-			Categories.Add(SerializedCategory.CategoryName);
+			Collections.Add(CollectionObjectSet.CollectionName);
 		}
 
-		return Categories;
+		return Collections;
 	}
 
 	return {};
diff --git a/Engine/Source/Editor/ObjectMixer/ObjectMixer/Private/Views/List/ObjectMixerEditorListRow.cpp b/Engine/Source/Editor/ObjectMixer/ObjectMixer/Private/Views/List/ObjectMixerEditorListRow.cpp
--- a/Engine/Source/Editor/ObjectMixer/ObjectMixer/Private/Views/List/ObjectMixerEditorListRow.cpp
+++ b/Engine/Source/Editor/ObjectMixer/ObjectMixer/Private/Views/List/ObjectMixerEditorListRow.cpp
@@ -41,7 +41,7 @@ UObjectMixerObjectFilter* FObjectMixerEditorListRow::GetObjectFilter() const
 	return nullptr;
 }
 
-bool FObjectMixerEditorListRow::IsObjectRefInSelectedCategories() const
+bool FObjectMixerEditorListRow::IsObjectRefInSelectedCollections() const
 {
 	check (ListViewPtr.IsValid());
 	
@@ -51,16 +51,16 @@ bool FObjectMixerEditorListRow::IsObjectRefInSelectedCategories() const
 		{
 			if (const TSharedPtr<FObjectMixerEditorMainPanel> MainPanel = ListModel->GetMainPanelModel().Pin())
 			{
-				const TSet<FName>& CategorySelection = MainPanel->GetCurrentCategorySelection();
+				const TSet<FName>& CollectionSelection = MainPanel->GetCurrentCollectionSelection();
 
-				if (CategorySelection.Num() == 0)
+				if (CollectionSelection.Num() == 0)
 				{
 					return true;
 				}
 
-				const TSet<FName>& ObjectAssignedCategories = MainPanel->GetCategoriesForObject(GetObject());
+				const TSet<FName>& ObjectAssignedCollections = MainPanel->GetCollectionsForObject(GetObject());
 
-				const TSet<FName> Intersection = CategorySelection.Intersect(ObjectAssignedCategories);
+				const TSet<FName> Intersection = CollectionSelection.Intersect(ObjectAssignedCollections);
 
 				return Intersection.Num() > 0;
 			}
diff --git a/Engine/Source/Editor/ObjectMixer/ObjectMixer/Private/Views/List/SObjectMixerEditorList.cpp b/Engine/Source/Editor/ObjectMixer/ObjectMixer/Private/Views/List/SObjectMixerEditorList.cpp
--- a/Engine/Source/Editor/ObjectMixer/ObjectMixer/Private/Views/List/SObjectMixerEditorList.cpp
+++ b/Engine/Source/Editor/ObjectMixer/ObjectMixer/Private/Views/List/SObjectMixerEditorList.cpp
@@ -450,18 +450,18 @@ TSharedRef<SWidget> SObjectMixerEditorList::GenerateHeaderRowContextMenu() const
 	
 	MenuBuilder.AddSearchWidget();
 
-	FName LastCategoryName = NAME_None;
+	FName LastPropertyCategoryName = NAME_None;
 
 	for (const FListViewColumnInfo& ColumnInfo : ListViewColumns)
 	{
-		const FName& CategoryName = ColumnInfo.CategoryName;
+		const FName& PropertyCategoryName = ColumnInfo.PropertyCategoryName;
 
-		if (!CategoryName.IsEqual(LastCategoryName))
+		if (!PropertyCategoryName.IsEqual(LastPropertyCategoryName))
 		{
-			LastCategoryName = CategoryName;
+			LastPropertyCategoryName = PropertyCategoryName;
 			
 			MenuBuilder.EndSection();
-            MenuBuilder.BeginSection(LastCategoryName, FText::FromName(LastCategoryName));
+            MenuBuilder.BeginSection(LastPropertyCategoryName, FText::FromName(LastPropertyCategoryName));
 		}
 		
 		const FName& PropertyName = ColumnInfo.PropertyName;
@@ -546,10 +546,10 @@ bool SObjectMixerEditorList::AddUniquePropertyColumnsToHeaderRow(
 				})
 			)
 		{
-			FString CategoryName = "Generated Properties";
+			FString PropertyCategoryName = "Generated Properties";
 			if (const FString* CategoryMeta = Property->FindMetaData("Category"))
 			{
-				CategoryName = *CategoryMeta;
+				PropertyCategoryName = *CategoryMeta;
 			}
 			
 			ListViewColumns.Add(
@@ -557,7 +557,7 @@ bool SObjectMixerEditorList::AddUniquePropertyColumnsToHeaderRow(
 					Property, PropertyName,
 					Property->GetDisplayNameText(),
 					EListViewColumnType::PropertyGenerated,
-					*CategoryName,
+					*PropertyCategoryName,
 					true, true, false
 				}
 			);
@@ -648,10 +648,10 @@ TSharedPtr<SHeaderRow> SObjectMixerEditorList::GenerateHeaderRow()
 		return A.PropertyDisplayText.ToString() < B.PropertyDisplayText.ToString();
 	});
 
-	// Alphabetical sort by Category Name
+	// Alphabetical sort by Property Category Name
 	ListViewColumns.StableSort([](const FListViewColumnInfo& A, const FListViewColumnInfo& B)
 	{
-		return A.CategoryName.LexicalLess(B.CategoryName);
+		return A.PropertyCategoryName.LexicalLess(B.PropertyCategoryName);
 	});
 
 	// Add Built-in Columns to beginning
@@ -1113,7 +1113,7 @@ void SObjectMixerEditorList::GenerateTreeView()
 
 	if (TSharedPtr<FObjectMixerEditorMainPanel> MainPanel = ListModelPtr.Pin()->GetMainPanelModel().Pin())
 	{
-		MainPanel->RebuildCategorySelector();
+		MainPanel->RebuildCollectionSelector();
 	}
 
 	RefreshList();
diff --git a/Engine/Source/Editor/ObjectMixer/ObjectMixer/Private/Views/MainPanel/ObjectMixerEditorMainPanel.cpp b/Engine/Source/Editor/ObjectMixer/ObjectMixer/Private/Views/MainPanel/ObjectMixerEditorMainPanel.cpp
--- a/Engine/Source/Editor/ObjectMixer/ObjectMixer/Private/Views/MainPanel/ObjectMixerEditorMainPanel.cpp
+++ b/Engine/Source/Editor/ObjectMixer/ObjectMixer/Private/Views/MainPanel/ObjectMixerEditorMainPanel.cpp
@@ -44,9 +44,9 @@ void FObjectMixerEditorMainPanel::RefreshList() const
 	}
 }
 
-void FObjectMixerEditorMainPanel::RebuildCategorySelector()
+void FObjectMixerEditorMainPanel::RebuildCollectionSelector()
 {
-	MainPanelWidget->RebuildCategorySelector();
+	MainPanelWidget->RebuildCollectionSelector();
 }
 
 FString FObjectMixerEditorMainPanel::GetSearchStringFromSearchInputField() const
@@ -97,85 +97,85 @@ const TArray<TSharedRef<IObjectMixerEditorListFilter>>& FObjectMixerEditorMainPa
 	return MainPanelWidget->GetShowFilters();	
 }
 
-void FObjectMixerEditorMainPanel::AddObjectsToCategory(const FName& CategoryName, const TSet<FSoftObjectPath>& ObjectsToAdd) const
+void FObjectMixerEditorMainPanel::AddObjectsToCollection(const FName& CollectionName, const TSet<FSoftObjectPath>& ObjectsToAdd) const
 {
 	if (const TSubclassOf<UObjectMixerObjectFilter> Filter = GetObjectFilterClass())
 	{
 		if (UObjectMixerEditorSerializedData* Settings = GetMutableDefault<UObjectMixerEditorSerializedData>())
 		{
-			Settings->AddObjectsToCategory(Filter->GetFName(), CategoryName, ObjectsToAdd);
+			Settings->AddObjectsToCollection(Filter->GetFName(), CollectionName, ObjectsToAdd);
 
-			OnObjectMixerCategoryMapChanged.Broadcast();
+			OnObjectMixerCollectionMapChanged.Broadcast();
 		}
 	}
 }
 
-void FObjectMixerEditorMainPanel::RemoveObjectsFromCategory(const FName& CategoryName, const TSet<FSoftObjectPath>& ObjectsToRemove) const
+void FObjectMixerEditorMainPanel::RemoveObjectsFromCollection(const FName& CollectionName, const TSet<FSoftObjectPath>& ObjectsToRemove) const
 {
 	if (const TSubclassOf<UObjectMixerObjectFilter> Filter = GetObjectFilterClass())
 	{
 		if (UObjectMixerEditorSerializedData* Settings = GetMutableDefault<UObjectMixerEditorSerializedData>())
 		{
-			Settings->RemoveObjectsFromCategory(Filter->GetFName(), CategoryName, ObjectsToRemove);
+			Settings->RemoveObjectsFromCollection(Filter->GetFName(), CollectionName, ObjectsToRemove);
 
-			OnObjectMixerCategoryMapChanged.Broadcast();
+			OnObjectMixerCollectionMapChanged.Broadcast();
 		}
 	}
 }
 
-void FObjectMixerEditorMainPanel::RemoveCategory(const FName& CategoryName) const
+void FObjectMixerEditorMainPanel::RemoveCollection(const FName& CollectionName) const
 {
 	if (const TSubclassOf<UObjectMixerObjectFilter> Filter = GetObjectFilterClass())
 	{
 		if (UObjectMixerEditorSerializedData* Settings = GetMutableDefault<UObjectMixerEditorSerializedData>())
 		{
-			Settings->RemoveCategory(Filter->GetFName(), CategoryName);
+			Settings->RemoveCollection(Filter->GetFName(), CollectionName);
 
-			OnObjectMixerCategoryMapChanged.Broadcast();
+			OnObjectMixerCollectionMapChanged.Broadcast();
 		}
 	}
 }
 
-bool FObjectMixerEditorMainPanel::IsObjectInCategory(const FName& CategoryName, const FSoftObjectPath& InObject) const
+bool FObjectMixerEditorMainPanel::IsObjectInCollection(const FName& CollectionName, const FSoftObjectPath& InObject) const
 {
 	if (const TSubclassOf<UObjectMixerObjectFilter> Filter = GetObjectFilterClass())
 	{
 		if (UObjectMixerEditorSerializedData* Settings = GetMutableDefault<UObjectMixerEditorSerializedData>())
 		{
-			return Settings->IsObjectInCategory(Filter->GetFName(), CategoryName, InObject);
+			return Settings->IsObjectInCollection(Filter->GetFName(), CollectionName, InObject);
 		}
 	}
 
 	return false;
 }
 
-TSet<FName> FObjectMixerEditorMainPanel::GetCategoriesForObject(const FSoftObjectPath& InObject) const
+TSet<FName> FObjectMixerEditorMainPanel::GetCollectionsForObject(const FSoftObjectPath& InObject) const
 {
 	if (const TSubclassOf<UObjectMixerObjectFilter> Filter = GetObjectFilterClass())
 	{
 		if (UObjectMixerEditorSerializedData* Settings = GetMutableDefault<UObjectMixerEditorSerializedData>())
 		{
-			return Settings->GetCategoriesForObject(Filter->GetFName(), InObject);
+			return Settings->GetCollectionsForObject(Filter->GetFName(), InObject);
 		}
 	}
 
 	return {};
 }
 
-TSet<FName> FObjectMixerEditorMainPanel::GetAllCategories() const
+TSet<FName> FObjectMixerEditorMainPanel::GetAllCollections() const
 {
 	if (const TSubclassOf<UObjectMixerObjectFilter> Filter = GetObjectFilterClass())
 	{
 		if (UObjectMixerEditorSerializedData* Settings = GetMutableDefault<UObjectMixerEditorSerializedData>())
 		{
-			return Settings->GetAllCategories(Filter->GetFName());
+			return Settings->GetAllCollections(Filter->GetFName());
 		}
 	}
 
 	return {};
 }
 
-const TSet<FName>& FObjectMixerEditorMainPanel::GetCurrentCategorySelection() const
+const TSet<FName>& FObjectMixerEditorMainPanel::GetCurrentCollectionSelection() const
 {
-	return MainPanelWidget->GetCurrentCategorySelection();
+	return MainPanelWidget->GetCurrentCollectionSelection();
 }
diff --git a/Engine/Source/Editor/ObjectMixer/ObjectMixer/Private/Views/MainPanel/SObjectMixerEditorMainPanel.cpp b/Engine/Source/Editor/ObjectMixer/ObjectMixer/Private/Views/MainPanel/SObjectMixerEditorMainPanel.cpp
--- a/Engine/Source/Editor/ObjectMixer/ObjectMixer/Private/Views/MainPanel/SObjectMixerEditorMainPanel.cpp
+++ b/Engine/Source/Editor/ObjectMixer/ObjectMixer/Private/Views/MainPanel/SObjectMixerEditorMainPanel.cpp
@@ -15,7 +15,7 @@
 #include "SPositiveActionButton.h"
 #include "Styling/StyleColors.h"
 #include "Views/List/ObjectMixerEditorListFilters/IObjectMixerEditorListFilter.h"
-#include "Views/List/ObjectMixerEditorListFilters/ObjectMixerEditorListFilter_Category.h"
+#include "Views/List/ObjectMixerEditorListFilters/ObjectMixerEditorListFilter_Collection.h"
 #include "Widgets/Input/SCheckBox.h"
 #include "Widgets/Input/SComboButton.h"
 #include "Widgets/Input/SMultiLineEditableTextBox.h"
@@ -48,7 +48,7 @@ void SObjectMixerEditorMainPanel::Construct(
 		.Padding(8, 2, 8, 7)
 		.AutoHeight()
 		[
-			SAssignNew(CategorySelectorBox, SWrapBox)
+			SAssignNew(CollectionSelectorBox, SWrapBox)
 			.UseAllottedSize(true)
 			.InnerSlotPadding(FVector2D(4,4))
 		]
@@ -61,9 +61,9 @@ void SObjectMixerEditorMainPanel::Construct(
 		]
 	];
 
-	InMainPanel->GetOnObjectMixerCategoryMapChanged().AddRaw(this, &SObjectMixerEditorMainPanel::RebuildCategorySelector);
+	InMainPanel->GetOnObjectMixerCollectionMapChanged().AddRaw(this, &SObjectMixerEditorMainPanel::RebuildCollectionSelector);
 
-	ShowFilters.Add(MakeShared<FObjectMixerEditorListFilter_Category>());
+	ShowFilters.Add(MakeShared<FObjectMixerEditorListFilter_Collection>());
 }
 
 TSharedRef<SWidget> SObjectMixerEditorMainPanel::GenerateToolbar()
@@ -466,75 +466,75 @@ void SObjectMixerEditorMainPanel::ToggleFilterActive(const FString& FilterName)
 	}
 }
 
-const TSet<FName>& SObjectMixerEditorMainPanel::GetCurrentCategorySelection()
+const TSet<FName>& SObjectMixerEditorMainPanel::GetCurrentCollectionSelection()
 {
-	return CurrentCategorySelection;
+	return CurrentCollectionSelection;
 }
 
-void SObjectMixerEditorMainPanel::RebuildCategorySelector()
+void SObjectMixerEditorMainPanel::RebuildCollectionSelector()
 {
 	check(MainPanelModel.IsValid());
 	
-	CategorySelectorBox->ClearChildren();
-	CategorySelectorBox->SetVisibility(EVisibility::Collapsed);
+	CollectionSelectorBox->ClearChildren();
+	CollectionSelectorBox->SetVisibility(EVisibility::Collapsed);
 
-	TArray<FName> AllCategories = MainPanelModel.Pin()->GetAllCategories().Array();
+	TArray<FName> AllCollections = MainPanelModel.Pin()->GetAllCollections().Array();
 	
-	if (AllCategories.IsEmpty())
+	if (AllCollections.IsEmpty())
 	{
 		// we've selected something that has no sections - rather than show just "All", hide the box
-		ResetCurrentCategorySelection();
+		ResetCurrentCollectionSelection();
 		return;
 	}
 
-	auto CreateSection = [this](FName CategoryName) 
+	auto CreateSection = [this](FName CollectionName) 
 	{
 		return SNew(SBox)
 			.Padding(FMargin(0))
 			[
 				SNew(SCheckBox)
 				.Style(FAppStyle::Get(), "DetailsView.SectionButton")
-				.OnCheckStateChanged(this, &SObjectMixerEditorMainPanel::OnCategoryCheckedChanged, CategoryName)
-				.IsChecked(this, &SObjectMixerEditorMainPanel::IsCategoryChecked, CategoryName)
+				.OnCheckStateChanged(this, &SObjectMixerEditorMainPanel::OnCollectionCheckedChanged, CollectionName)
+				.IsChecked(this, &SObjectMixerEditorMainPanel::IsCollectionChecked, CollectionName)
 				[
 					SNew(STextBlock)
 					.TextStyle(FAppStyle::Get(), "SmallText")
-					.Text(FText::FromName(CategoryName))
+					.Text(FText::FromName(CollectionName))
 				]
 			];
 	};
 	
-	AllCategories.Sort([](FName A, FName B)
+	AllCollections.Sort([](FName A, FName B)
 		{
 			return A.LexicalLess(B);
 		});
 
-	for (const FName& Key : AllCategories)
+	for (const FName& Key : AllCollections)
 	{
-		CategorySelectorBox->AddSlot()
+		CollectionSelectorBox->AddSlot()
 		[
 			CreateSection(Key)
 		];
 	}
 
-	CategorySelectorBox->AddSlot()
+	CollectionSelectorBox->AddSlot()
 	[
 		CreateSection("All")
 	];
 
-	CategorySelectorBox->SetVisibility(EVisibility::Visible);
+	CollectionSelectorBox->SetVisibility(EVisibility::Visible);
 }
 
-void SObjectMixerEditorMainPanel::OnCategoryCheckedChanged(ECheckBoxState State, FName SectionName)
+void SObjectMixerEditorMainPanel::OnCollectionCheckedChanged(ECheckBoxState State, FName SectionName)
 {
 	check(MainPanelModel.IsValid());
 
-	// Remove category
+	// Remove Collection
 	if (FSlateApplication::Get().GetModifierKeys().IsAltDown())
 	{
-		MainPanelModel.Pin()->RemoveCategory(SectionName);
+		MainPanelModel.Pin()->RemoveCollection(SectionName);
 
-		CurrentCategorySelection.Remove(SectionName);
+		CurrentCollectionSelection.Remove(SectionName);
 	}
 	else
 	{
@@ -544,15 +544,15 @@ void SObjectMixerEditorMainPanel::OnCategoryCheckedChanged(ECheckBoxState State,
 		{
 			if (bIsControlDown)
 			{
-				CurrentCategorySelection.Remove(SectionName);
+				CurrentCollectionSelection.Remove(SectionName);
 			}
 			else
 			{
-				CurrentCategorySelection.Reset();
+				CurrentCollectionSelection.Reset();
 
 				if (SectionName != "All")
 				{
-					CurrentCategorySelection.Add(SectionName);
+					CurrentCollectionSelection.Add(SectionName);
 				}
 			}
 		}
@@ -560,12 +560,12 @@ void SObjectMixerEditorMainPanel::OnCategoryCheckedChanged(ECheckBoxState State,
 		{
 			if (!bIsControlDown)
 			{
-				CurrentCategorySelection.Reset();
+				CurrentCollectionSelection.Reset();
 			}
 
 			if (SectionName != "All")
 			{
-				CurrentCategorySelection.Add(SectionName);
+				CurrentCollectionSelection.Add(SectionName);
 			}
 		}
 	}
@@ -573,21 +573,21 @@ void SObjectMixerEditorMainPanel::OnCategoryCheckedChanged(ECheckBoxState State,
 	MainPanelModel.Pin()->GetEditorListModel().Pin()->EvaluateIfRowsPassFilters();
 }
 
-ECheckBoxState SObjectMixerEditorMainPanel::IsCategoryChecked(FName Section) const
+ECheckBoxState SObjectMixerEditorMainPanel::IsCollectionChecked(FName Section) const
 {
-	if (CurrentCategorySelection.IsEmpty() && Section == "All")
+	if (CurrentCollectionSelection.IsEmpty() && Section == "All")
 	{
 		return ECheckBoxState::Checked;
 	}
 
-	return CurrentCategorySelection.Contains(Section) ? ECheckBoxState::Checked : ECheckBoxState::Unchecked;
+	return CurrentCollectionSelection.Contains(Section) ? ECheckBoxState::Checked : ECheckBoxState::Unchecked;
 }
 
 SObjectMixerEditorMainPanel::~SObjectMixerEditorMainPanel()
 {
 	if (MainPanelModel.IsValid())
 	{
-		MainPanelModel.Pin()->GetOnObjectMixerCategoryMapChanged().RemoveAll(this);
+		MainPanelModel.Pin()->GetOnObjectMixerCollectionMapChanged().RemoveAll(this);
 		MainPanelModel.Reset();
 	}
 }
diff --git a/Engine/Source/Editor/ObjectMixer/ObjectMixer/Private/Views/Widgets/ObjectMixerEditorListMenuContext.cpp b/Engine/Source/Editor/ObjectMixer/ObjectMixer/Private/Views/Widgets/ObjectMixerEditorListMenuContext.cpp
--- a/Engine/Source/Editor/ObjectMixer/ObjectMixer/Private/Views/Widgets/ObjectMixerEditorListMenuContext.cpp
+++ b/Engine/Source/Editor/ObjectMixer/ObjectMixer/Private/Views/Widgets/ObjectMixerEditorListMenuContext.cpp
@@ -2,6 +2,8 @@
 
 #include "Views/Widgets/ObjectMixerEditorListMenuContext.h"
 
+#include "LevelEditor.h"
+#include "LevelEditorMenuContext.h"
 #include "Framework/MultiBox/MultiBoxDefs.h"
 #include "ToolMenu.h"
 #include "ToolMenus.h"
@@ -37,19 +39,25 @@ TSharedPtr<SWidget> UObjectMixerEditorListMenuContext::BuildContextMenu(const FO
 
 	const FToolMenuContext Context(ContextObject);
 
-	// Build up the menu for a selection
+	// todo: Uncomment all code when Level Editor Changes from 21454357 are pushed
+	// First get scene outliner-style context menu
+	// FLevelEditorModule& LevelEditorModule = FModuleManager::LoadModuleChecked<FLevelEditorModule>("LevelEditor");
+	// TSharedPtr<ILevelEditor> LevelEditorInstancePtr = LevelEditorModule.GetLevelEditorInstance().Pin();
+	//
+	// UToolMenu* BaseMenu = LevelEditorInstancePtr->GenerateLevelEditorContextMenu(
+	// 		ELevelEditorMenuContext::SceneOutliner, nullptr, FTypedElementHandle());
+
+	// Now generate custom sections to append
 	UToolMenus* ToolMenus = UToolMenus::Get();
-	UToolMenu* Menu = ToolMenus->GenerateMenu(DefaultContextBaseMenuName, Context);
+	UToolMenu* CustomMenu = ToolMenus->GenerateMenu(DefaultContextBaseMenuName, Context);
 
-	for (const FToolMenuSection& Section : Menu->Sections)
-	{
-		if (Section.Blocks.Num() > 0)
-		{
-			return ToolMenus->GenerateWidget(Menu);
-		}
-	}
+	// for (const FToolMenuSection& Section : CustomMenu->Sections)
+	// {
+	// 	BaseMenu->Sections.Add(Section);
+	// }
 
-	return nullptr;
+	//return ToolMenus->GenerateWidget(BaseMenu);
+	return ToolMenus->GenerateWidget(CustomMenu);
 }
 
 void UObjectMixerEditorListMenuContext::RegisterContextMenu()
@@ -60,7 +68,7 @@ void UObjectMixerEditorListMenuContext::RegisterContextMenu()
 	{
 		UToolMenu* Menu = ToolMenus->RegisterMenu(DefaultContextBaseMenuName);
 	
-		Menu->AddDynamicSection("DynamicCategorizationSection", FNewToolMenuDelegate::CreateLambda(
+		Menu->AddDynamicSection("DynamicCollectionsSection", FNewToolMenuDelegate::CreateLambda(
 			[](UToolMenu* InMenu)
 			{
 				UObjectMixerEditorListMenuContext* Context = InMenu->FindContext<UObjectMixerEditorListMenuContext>();
@@ -69,13 +77,13 @@ void UObjectMixerEditorListMenuContext::RegisterContextMenu()
 					return;
 				}
 	
-				FToolMenuSection& Section = InMenu->FindOrAddSection("Categorization");
-				Section.Label = LOCTEXT("CategorizationSectionName", "Categorization");
+				FToolMenuSection& Section = InMenu->FindOrAddSection("Collections");
+				Section.Label = LOCTEXT("CollectionsSectionName", "Collections");
 			
 				Section.AddSubMenu(
-					"SelectCategorySubMenu",
-					LOCTEXT("SelectCategorySubmenu", "Select or Add Category"),
-					LOCTEXT("SelectCategorySubmenu_Tooltip", "Select the category to which you wish to assign this object."),
+					"SelectCollectionsSubMenu",
+					LOCTEXT("SelectCollectionsSubmenu", "Select or Add Collection"),
+					LOCTEXT("SelectCollectionsSubmenu_Tooltip", "Select the collection to which you wish to assign this object."),
 					FNewToolMenuDelegate::CreateLambda(
 						[Context](UToolMenu* Menu)
 						{
@@ -89,40 +97,40 @@ void UObjectMixerEditorListMenuContext::RegisterContextMenu()
 											.Padding(8, 0)
 											[
 												SAssignNew(Context->EditableText, SEditableTextBox)
-												.HintText(LOCTEXT("NewCategoryEditableHintText", "Enter a new category name..."))
+												.HintText(LOCTEXT("NewCollectionEditableHintText", "Enter a new collection name..."))
 												.OnTextCommitted_Static(&UObjectMixerEditorListMenuContext::OnTextCommitted, Context)
 												.OnTextChanged_Static(&UObjectMixerEditorListMenuContext::OnTextChanged, Context)
 											]
 									;
 								}
 							);
-							Menu->AddMenuEntry("NewCategoryInput", Args);
+							Menu->AddMenuEntry("NewCollectionInput", Args);
 							
 							if (const TSharedPtr<FObjectMixerEditorMainPanel> MainPanel = Context->Data.MainPanelPtr.Pin())
 							{
-								if (TArray<FName> Categories = MainPanel->GetAllCategories().Array(); Categories.Num() > 0)
+								if (TArray<FName> Collections = MainPanel->GetAllCollections().Array(); Collections.Num() > 0)
 								{
-									FToolMenuSection& Section = Menu->FindOrAddSection("Categories");
-								   Section.Label = LOCTEXT("CategoriesSectionName", "Categories");
+									FToolMenuSection& Section = Menu->FindOrAddSection("Collections");
+								   Section.Label = LOCTEXT("CollectionsSectionName", "Collections");
 
-								   Categories.StableSort([](const FName& A, const FName B)
+								   Collections.StableSort([](const FName& A, const FName B)
 								   {
 									   return A.LexicalLess(B);
 								   });
 								
-								   for (const FName& Key : Categories)
+								   for (const FName& Key : Collections)
 								   {
 									   const FText KeyText = FText::FromName(Key);
 									
 									   Section.AddMenuEntry(
 										   Key,
 										   KeyText,
-										   FText::Format(LOCTEXT("AddObjectsToCategoryTooltipFormat", "Add selected to category '{0}'"), KeyText),
+										   FText::Format(LOCTEXT("AddObjectsToCollectionTooltipFormat", "Add selected to collection '{0}'"), KeyText),
 										   FSlateIcon(),
 										   FUIAction(
-											   FExecuteAction::CreateStatic(&UObjectMixerEditorListMenuContext::OnClickCategoryMenuEntry, Key, Context),
+											   FExecuteAction::CreateStatic(&UObjectMixerEditorListMenuContext::OnClickCollectionMenuEntry, Key, Context),
 											   FCanExecuteAction(),
-											   FIsActionChecked::CreateStatic(&UObjectMixerEditorListMenuContext::AreAllObjectsInCategory, Key, Context)
+											   FIsActionChecked::CreateStatic(&UObjectMixerEditorListMenuContext::AreAllObjectsInCollection, Key, Context)
 										   ),
 										   EUserInterfaceActionType::Check
 									   );
@@ -145,7 +153,7 @@ void UObjectMixerEditorListMenuContext::OnTextCommitted(const FText& InText, ETe
 	{
 		UE_LOG(LogTemp, Warning, TEXT("%hs: %s"), __FUNCTION__, *InText.ToString());
 
-		AddObjectsToCategory(*InText.ToString(), Context);
+		AddObjectsToCollection(*InText.ToString(), Context);
 	}
 }
 
@@ -159,19 +167,19 @@ void UObjectMixerEditorListMenuContext::OnTextChanged(const FText& InText, UObje
 	}
 }
 
-void UObjectMixerEditorListMenuContext::OnClickCategoryMenuEntry(const FName Key, UObjectMixerEditorListMenuContext* Context)
+void UObjectMixerEditorListMenuContext::OnClickCollectionMenuEntry(const FName Key, UObjectMixerEditorListMenuContext* Context)
 {
-	if (AreAllObjectsInCategory(Key, Context))
+	if (AreAllObjectsInCollection(Key, Context))
 	{
-		RemoveObjectsFromCategory(Key, Context);
+		RemoveObjectsFromCollection(Key, Context);
 	}
 	else
 	{
-		AddObjectsToCategory(Key, Context);
+		AddObjectsToCollection(Key, Context);
 	}
 }
 
-void UObjectMixerEditorListMenuContext::AddObjectsToCategory(const FName Key, UObjectMixerEditorListMenuContext* Context)
+void UObjectMixerEditorListMenuContext::AddObjectsToCollection(const FName Key, UObjectMixerEditorListMenuContext* Context)
 {
 	check(Context);
 
@@ -187,11 +195,11 @@ void UObjectMixerEditorListMenuContext::AddObjectsToCategory(const FName Key, UO
 			}
 		}
 		
-		MainPanel->AddObjectsToCategory(Key, ObjectPaths);
+		MainPanel->AddObjectsToCollection(Key, ObjectPaths);
 	}
 }
 
-void UObjectMixerEditorListMenuContext::RemoveObjectsFromCategory(const FName Key, UObjectMixerEditorListMenuContext* Context)
+void UObjectMixerEditorListMenuContext::RemoveObjectsFromCollection(const FName Key, UObjectMixerEditorListMenuContext* Context)
 {
 	check(Context);
 
@@ -207,15 +215,15 @@ void UObjectMixerEditorListMenuContext::RemoveObjectsFromCategory(const FName Ke
 			}
 		}
 		
-		MainPanel->RemoveObjectsFromCategory(Key, ObjectPaths);
+		MainPanel->RemoveObjectsFromCollection(Key, ObjectPaths);
 	}
 }
 
-bool UObjectMixerEditorListMenuContext::AreAllObjectsInCategory(const FName Key, UObjectMixerEditorListMenuContext* Context)
+bool UObjectMixerEditorListMenuContext::AreAllObjectsInCollection(const FName Key, UObjectMixerEditorListMenuContext* Context)
 {
 	check(Context);
 
-	bool bAreAllSelectedObjectsInCategory = false;
+	bool bAreAllSelectedObjectsInCollection = false;
 
 	if (const TSharedPtr<FObjectMixerEditorMainPanel> MainPanel = Context->Data.MainPanelPtr.Pin())
 	{
@@ -223,9 +231,9 @@ bool UObjectMixerEditorListMenuContext::AreAllObjectsInCategory(const FName Key,
 		{
 			if (const UObject* Object = Item->GetObject())
 			{
-				bAreAllSelectedObjectsInCategory = MainPanel->IsObjectInCategory(Key, Object);
+				bAreAllSelectedObjectsInCollection = MainPanel->IsObjectInCollection(Key, Object);
 
-				if (!bAreAllSelectedObjectsInCategory)
+				if (!bAreAllSelectedObjectsInCollection)
 				{
 					return false;
 				}
@@ -233,7 +241,7 @@ bool UObjectMixerEditorListMenuContext::AreAllObjectsInCategory(const FName Key,
 		}
 	}
 
-	return bAreAllSelectedObjectsInCategory;
+	return bAreAllSelectedObjectsInCollection;
 }
 
 #undef LOCTEXT_NAMESPACE
diff --git a/Engine/Source/Editor/ObjectMixer/ObjectMixer/Public/ObjectMixerEditorSerializedData.h b/Engine/Source/Editor/ObjectMixer/ObjectMixer/Public/ObjectMixerEditorSerializedData.h
--- a/Engine/Source/Editor/ObjectMixer/ObjectMixer/Public/ObjectMixerEditorSerializedData.h
+++ b/Engine/Source/Editor/ObjectMixer/ObjectMixer/Public/ObjectMixerEditorSerializedData.h
@@ -5,24 +5,24 @@
 #include "ObjectMixerEditorSerializedData.generated.h"
 
 USTRUCT()
-struct FObjectMixerCategoryObjectSet
+struct FObjectMixerCollectionObjectSet
 {
 	GENERATED_BODY()
 	
 	UPROPERTY()
-	FName CategoryName = NAME_None;
+	FName CollectionName = NAME_None;
 
 	UPROPERTY()
-	TSet<FSoftObjectPath> CategoryObjects = {};
+	TSet<FSoftObjectPath> CollectionObjects = {};
 
-	bool operator==(const FObjectMixerCategoryObjectSet& Other) const
+	bool operator==(const FObjectMixerCollectionObjectSet& Other) const
 	{
-		return CategoryName.IsEqual(Other.CategoryName);
+		return CollectionName.IsEqual(Other.CollectionName);
 	}
 
-	friend uint32 GetTypeHash (const FObjectMixerCategoryObjectSet& Other)
+	friend uint32 GetTypeHash (const FObjectMixerCollectionObjectSet& Other)
 	{
-		return GetTypeHash(Other.CategoryName);
+		return GetTypeHash(Other.CollectionName);
 	}
 };
 
@@ -35,7 +35,7 @@ struct FObjectMixerSerializationData
 	FName FilterClassName = NAME_None;
 
 	UPROPERTY()
-	TSet<FObjectMixerCategoryObjectSet> SerializedCategories = {};
+	TSet<FObjectMixerCollectionObjectSet> SerializedCollection = {};
 
 	bool operator==(const FObjectMixerSerializationData& Other) const
 	{
@@ -62,15 +62,15 @@ public:
 
 	FObjectMixerSerializationData* FindSerializationDataByFilterClassName(const FName& FilterClassName);
 
-	void AddObjectsToCategory(const FName& FilterClassName, const FName& CategoryName, const TSet<FSoftObjectPath>& ObjectsToAdd);
+	void AddObjectsToCollection(const FName& FilterClassName, const FName& CollectionName, const TSet<FSoftObjectPath>& ObjectsToAdd);
 
-	void RemoveObjectsFromCategory(const FName& FilterClassName, const FName& CategoryName, const TSet<FSoftObjectPath>& ObjectsToRemove);
+	void RemoveObjectsFromCollection(const FName& FilterClassName, const FName& CollectionName, const TSet<FSoftObjectPath>& ObjectsToRemove);
 
-	void RemoveCategory(const FName& FilterClassName, const FName& CategoryName);
+	void RemoveCollection(const FName& FilterClassName, const FName& CollectionName);
 
-	bool IsObjectInCategory(const FName& FilterClassName, const FName& CategoryName, const FSoftObjectPath& InObject);
+	bool IsObjectInCollection(const FName& FilterClassName, const FName& CollectionName, const FSoftObjectPath& InObject);
 
-	TSet<FName> GetCategoriesForObject(const FName& FilterClassName, const FSoftObjectPath& InObject);
+	TSet<FName> GetCollectionsForObject(const FName& FilterClassName, const FSoftObjectPath& InObject);
 
-	TSet<FName> GetAllCategories(const FName& FilterClassName);
+	TSet<FName> GetAllCollections(const FName& FilterClassName);
 };
diff --git a/Engine/Source/Editor/ObjectMixer/ObjectMixer/Public/Views/List/ObjectMixerEditorListFilters/ObjectMixerEditorListFilter_Category.h b/Engine/Source/Editor/ObjectMixer/ObjectMixer/Public/Views/List/ObjectMixerEditorListFilters/ObjectMixerEditorListFilter_Category.h
deleted file mode 100644
diff --git a/Engine/Source/Editor/ObjectMixer/ObjectMixer/Public/Views/List/ObjectMixerEditorListFilters/ObjectMixerEditorListFilter_Collection.h b/Engine/Source/Editor/ObjectMixer/ObjectMixer/Public/Views/List/ObjectMixerEditorListFilters/ObjectMixerEditorListFilter_Collection.h
new file mode 100644
--- /dev/null
+++ b/Engine/Source/Editor/ObjectMixer/ObjectMixer/Public/Views/List/ObjectMixerEditorListFilters/ObjectMixerEditorListFilter_Collection.h
@@ -0,0 +1,54 @@
+// Copyright Epic Games, Inc. All Rights Reserved.
+
+#pragma once
+
+#include "IObjectMixerEditorListFilter.h"
+
+#define LOCTEXT_NAMESPACE "ObjectMixerEditor"
+
+class OBJECTMIXEREDITOR_API FObjectMixerEditorListFilter_Collection : public IObjectMixerEditorListFilter
+{
+public:
+
+	FObjectMixerEditorListFilter_Collection()
+	{
+		SetFilterMatchType(EObjectMixerEditorListFilterMatchType::MatchAll);
+	}
+
+	virtual FString GetFilterName() const override
+	{
+		return "ObjectMixerCollectionListFilter";
+	}
+
+	virtual bool IsToggleable() const override
+	{
+		return false;
+	}
+
+	virtual FText GetFilterButtonLabel() const override
+	{
+		return LOCTEXT("ShowSourceTextFilterFormat", "Show Collections");
+	}
+
+	virtual FText GetFilterButtonToolTip() const override
+	{
+		return LOCTEXT("ShowSourceTextFilterTooltipFormat", "Show rows that are assigned to the selected collections");
+	}
+
+	virtual bool DoesItemPassFilter(const FObjectMixerEditorListRowPtr& InItem) const override
+	{
+		if (InItem.IsValid())
+		{
+			if (InItem->GetRowType() == FObjectMixerEditorListRow::ContainerObject ||
+				InItem->GetRowType() == FObjectMixerEditorListRow::MatchingObject)
+			{
+				return InItem->IsObjectRefInSelectedCollections();
+			}
+		}
+
+		return false;
+	}
+
+};
+
+#undef LOCTEXT_NAMESPACE
diff --git a/Engine/Source/Editor/ObjectMixer/ObjectMixer/Public/Views/List/ObjectMixerEditorListRow.h b/Engine/Source/Editor/ObjectMixer/ObjectMixer/Public/Views/List/ObjectMixerEditorListRow.h
--- a/Engine/Source/Editor/ObjectMixer/ObjectMixer/Public/Views/List/ObjectMixerEditorListRow.h
+++ b/Engine/Source/Editor/ObjectMixer/ObjectMixer/Public/Views/List/ObjectMixerEditorListRow.h
@@ -53,7 +53,7 @@ struct OBJECTMIXEREDITOR_API FObjectMixerEditorListRow final : TSharedFromThis<F
 
 	UObjectMixerObjectFilter* GetObjectFilter() const;
 
-	bool IsObjectRefInSelectedCategories() const;
+	bool IsObjectRefInSelectedCollections() const;
 
 	[[nodiscard]] EObjectMixerEditorListRowType GetRowType() const;
 
diff --git a/Engine/Source/Editor/ObjectMixer/ObjectMixer/Public/Views/List/SObjectMixerEditorList.h b/Engine/Source/Editor/ObjectMixer/ObjectMixer/Public/Views/List/SObjectMixerEditorList.h
--- a/Engine/Source/Editor/ObjectMixer/ObjectMixer/Public/Views/List/SObjectMixerEditorList.h
+++ b/Engine/Source/Editor/ObjectMixer/ObjectMixer/Public/Views/List/SObjectMixerEditorList.h
@@ -45,7 +45,7 @@ struct FListViewColumnInfo
 	EListViewColumnType PropertyType = EListViewColumnType::PropertyGenerated;
 
 	UPROPERTY()
-	FName CategoryName = NAME_None;
+	FName PropertyCategoryName = NAME_None;
 
 	UPROPERTY()
 	bool bIsDesiredForDisplay = false;
diff --git a/Engine/Source/Editor/ObjectMixer/ObjectMixer/Public/Views/MainPanel/ObjectMixerEditorMainPanel.h b/Engine/Source/Editor/ObjectMixer/ObjectMixer/Public/Views/MainPanel/ObjectMixerEditorMainPanel.h
--- a/Engine/Source/Editor/ObjectMixer/ObjectMixer/Public/Views/MainPanel/ObjectMixerEditorMainPanel.h
+++ b/Engine/Source/Editor/ObjectMixer/ObjectMixer/Public/Views/MainPanel/ObjectMixerEditorMainPanel.h
@@ -11,7 +11,7 @@
 class FObjectMixerEditorList;
 class SObjectMixerEditorMainPanel;
 
-DECLARE_MULTICAST_DELEGATE(FOnObjectMixerCategoryMapChanged)
+DECLARE_MULTICAST_DELEGATE(FOnObjectMixerCollectionMapChanged)
 
 class OBJECTMIXEREDITOR_API FObjectMixerEditorMainPanel : public TSharedFromThis<FObjectMixerEditorMainPanel>
 {
@@ -44,7 +44,7 @@ public:
 		return EditorListModel;
 	}
 
-	void RebuildCategorySelector();
+	void RebuildCollectionSelector();
 
 	FString GetSearchStringFromSearchInputField() const;
 
@@ -146,26 +146,26 @@ public:
 		return ModuleName;
 	}
 
-	// User Categorization
+	// User Collections
 
 	/**
-	 * Add set of objects to a category in the map, or create a new category if one does not exist,
+	 * Add set of objects to a collection in the map, or create a new collection if one does not exist,
 	 */
-	void AddObjectsToCategory(const FName& CategoryName, const TSet<FSoftObjectPath>& ObjectsToAdd) const;
-	void RemoveObjectsFromCategory(const FName& CategoryName, const TSet<FSoftObjectPath>& ObjectsToRemove) const;
-	void RemoveCategory(const FName& CategoryName) const;
-	bool IsObjectInCategory(const FName& CategoryName, const FSoftObjectPath& InObject) const;
-	TSet<FName> GetCategoriesForObject(const FSoftObjectPath& InObject) const;
-	TSet<FName> GetAllCategories() const;
-	FOnObjectMixerCategoryMapChanged& GetOnObjectMixerCategoryMapChanged()
+	void AddObjectsToCollection(const FName& CollectionName, const TSet<FSoftObjectPath>& ObjectsToAdd) const;
+	void RemoveObjectsFromCollection(const FName& CollectionName, const TSet<FSoftObjectPath>& ObjectsToRemove) const;
+	void RemoveCollection(const FName& CollectionName) const;
+	bool IsObjectInCollection(const FName& CollectionName, const FSoftObjectPath& InObject) const;
+	TSet<FName> GetCollectionsForObject(const FSoftObjectPath& InObject) const;
+	TSet<FName> GetAllCollections() const;
+	FOnObjectMixerCollectionMapChanged& GetOnObjectMixerCollectionMapChanged()
 	{
-		return OnObjectMixerCategoryMapChanged;
+		return OnObjectMixerCollectionMapChanged;
 	}
 
 	/**
-	 * Returns the categories selected by the user. If the set is empty, consider "All" categories to be selected.
+	 * Returns the collections selected by the user. If the set is empty, consider "All" collections to be selected.
 	 */
-	const TSet<FName>& GetCurrentCategorySelection() const;
+	const TSet<FName>& GetCurrentCollectionSelection() const;
 
 private:
 
@@ -189,5 +189,5 @@ private:
 
 	FName ModuleName = NAME_None;
 
-	FOnObjectMixerCategoryMapChanged OnObjectMixerCategoryMapChanged;
+	FOnObjectMixerCollectionMapChanged OnObjectMixerCollectionMapChanged;
 };
diff --git a/Engine/Source/Editor/ObjectMixer/ObjectMixer/Public/Views/MainPanel/SObjectMixerEditorMainPanel.h b/Engine/Source/Editor/ObjectMixer/ObjectMixer/Public/Views/MainPanel/SObjectMixerEditorMainPanel.h
--- a/Engine/Source/Editor/ObjectMixer/ObjectMixer/Public/Views/MainPanel/SObjectMixerEditorMainPanel.h
+++ b/Engine/Source/Editor/ObjectMixer/ObjectMixer/Public/Views/MainPanel/SObjectMixerEditorMainPanel.h
@@ -42,9 +42,9 @@ public:
 		return ShowFilters;	
 	}
 
-	const TSet<FName>& GetCurrentCategorySelection();
+	const TSet<FName>& GetCurrentCollectionSelection();
 	
-	void RebuildCategorySelector();
+	void RebuildCollectionSelector();
 
 	virtual ~SObjectMixerEditorMainPanel() override;
 
@@ -66,16 +66,16 @@ private:
 
 	void OnSearchTextChanged(const FText& Text);
 
-	// User Categorization
+	// User Collections
 	
-	TSharedPtr<class SWrapBox> CategorySelectorBox;
-	TSet<FName> CurrentCategorySelection;
+	TSharedPtr<class SWrapBox> CollectionSelectorBox;
+	TSet<FName> CurrentCollectionSelection;
 
-	void ResetCurrentCategorySelection()
+	void ResetCurrentCollectionSelection()
 	{
-		CurrentCategorySelection.Reset();
+		CurrentCollectionSelection.Reset();
 	}
 	
-	void OnCategoryCheckedChanged(ECheckBoxState State, FName SectionName);
-	ECheckBoxState IsCategoryChecked(FName Section) const;
+	void OnCollectionCheckedChanged(ECheckBoxState State, FName SectionName);
+	ECheckBoxState IsCollectionChecked(FName Section) const;
 };
diff --git a/Engine/Source/Editor/ObjectMixer/ObjectMixer/Public/Views/Widgets/ObjectMixerEditorListMenuContext.h b/Engine/Source/Editor/ObjectMixer/ObjectMixer/Public/Views/Widgets/ObjectMixerEditorListMenuContext.h
--- a/Engine/Source/Editor/ObjectMixer/ObjectMixer/Public/Views/Widgets/ObjectMixerEditorListMenuContext.h
+++ b/Engine/Source/Editor/ObjectMixer/ObjectMixer/Public/Views/Widgets/ObjectMixerEditorListMenuContext.h
@@ -35,10 +35,10 @@ private:
 	static void OnTextCommitted(const FText& InText, ETextCommit::Type InCommitType, UObjectMixerEditorListMenuContext* Context);
 	static void OnTextChanged(const FText& InText, UObjectMixerEditorListMenuContext* Context);
 	
-	static void OnClickCategoryMenuEntry(const FName Key, UObjectMixerEditorListMenuContext* Context);
-	static void AddObjectsToCategory(const FName Key, UObjectMixerEditorListMenuContext* Context);
-	static void RemoveObjectsFromCategory(const FName Key, UObjectMixerEditorListMenuContext* Context);
-	static bool AreAllObjectsInCategory(const FName Key, UObjectMixerEditorListMenuContext* Context);
+	static void OnClickCollectionMenuEntry(const FName Key, UObjectMixerEditorListMenuContext* Context);
+	static void AddObjectsToCollection(const FName Key, UObjectMixerEditorListMenuContext* Context);
+	static void RemoveObjectsFromCollection(const FName Key, UObjectMixerEditorListMenuContext* Context);
+	static bool AreAllObjectsInCollection(const FName Key, UObjectMixerEditorListMenuContext* Context);
 
 	TSharedPtr<class SEditableTextBox> EditableText;
 };
