Object mixer:

-Remove sorting mode from columns
-Add tooltips to Solo and Visibility columns

diff --git a/Engine/Source/Editor/ObjectMixer/ObjectMixer/Private/Views/List/SObjectMixerEditorList.cpp b/Engine/Source/Editor/ObjectMixer/ObjectMixer/Private/Views/List/SObjectMixerEditorList.cpp
--- a/Engine/Source/Editor/ObjectMixer/ObjectMixer/Private/Views/List/SObjectMixerEditorList.cpp
+++ b/Engine/Source/Editor/ObjectMixer/ObjectMixer/Private/Views/List/SObjectMixerEditorList.cpp
@@ -558,7 +558,7 @@ bool SObjectMixerEditorList::AddUniquePropertyColumnsToHeaderRow(
 					Property->GetDisplayNameText(),
 					EListViewColumnType::PropertyGenerated,
 					*PropertyCategoryName,
-					true, true, false
+					true, false, false
 				}
 			);
 
@@ -576,7 +576,7 @@ void SObjectMixerEditorList::AddBuiltinColumnsToHeaderRow()
 			nullptr, ItemNameColumnName,
 			LOCTEXT("ItemNameHeaderText", "Name"),
 			EListViewColumnType::BuiltIn, "Built-In",
-			true, true,
+			true, false,
 			false, 1.0f, 1.7f
 		}, 0
 	);
@@ -673,6 +673,7 @@ TSharedPtr<SHeaderRow> SObjectMixerEditorList::GenerateHeaderRow()
 				.DefaultLabel(ColumnInfo.PropertyDisplayText)
 				.ToolTipText(Tooltip)
 				.HAlignHeader(EHorizontalAlignment::HAlign_Left)
+				//.bHideHeaderMenuButton(false) todo: uncomment this line when 21478974 is pushed 
 			;
 
 			if (ColumnInfo.bUseFixedWidth)
@@ -707,6 +708,7 @@ TSharedPtr<SHeaderRow> SObjectMixerEditorList::GenerateHeaderRow()
 						SNew(SImage)
 						.ColorAndOpacity(FSlateColor::UseForeground())
 						.Image(FAppStyle::Get().GetBrush("Level.VisibleIcon16x"))
+						.ToolTipText(LOCTEXT("VisibilityColumnTooltip", "Visibility"))
 					]
 				];
 			}
@@ -722,6 +724,7 @@ TSharedPtr<SHeaderRow> SObjectMixerEditorList::GenerateHeaderRow()
 						SNew(SImage)
 						.ColorAndOpacity(FSlateColor::UseForeground())
 						.Image(FAppStyle::Get().GetBrush("MediaAsset.AssetActions.Solo.Small"))
+						.ToolTipText(LOCTEXT("SoloColumnTooltip", "Solo"))
 					]
 				];
 			}
