Object mixer:

-Add SCollectionSelectionButton widget (custom checkbox with drag and drop)
-Right-click collection button to remove it (with undo transaction)
-Implement drag and drop rows onto collection buttons
-Implement drag and drop collection buttons (for custom order)
-Remove SObjectMixerEditorTooltipWidget (unused)

diff --git a/Engine/Source/Editor/ObjectMixer/ObjectMixer/Private/ObjectMixerEditorModule.cpp b/Engine/Source/Editor/ObjectMixer/ObjectMixer/Private/ObjectMixerEditorModule.cpp
--- a/Engine/Source/Editor/ObjectMixer/ObjectMixer/Private/ObjectMixerEditorModule.cpp
+++ b/Engine/Source/Editor/ObjectMixer/ObjectMixer/Private/ObjectMixerEditorModule.cpp
@@ -10,6 +10,7 @@
 #include "Editor.h"
 #include "Framework/Docking/TabManager.h"
 #include "ISettingsModule.h"
+#include "ObjectMixerEditorSerializedData.h"
 #include "Misc/TransactionObjectEvent.h"
 #include "Modules/ModuleManager.h"
 #include "ToolMenus.h"
@@ -308,6 +309,9 @@ void FObjectMixerEditorModule::BindDelegates()
 	DelegateHandles.Add(FEditorDelegates::PostUndoRedo.AddLambda([this]()
 	{
 		RequestRebuildList();
+		
+		// Because we can undo/redo collection adds and removes, we need to save the data after each
+		GetMutableDefault<UObjectMixerEditorSerializedData>()->SaveConfig();
 	}));
 
 	DelegateHandles.Add(FCoreUObjectDelegates::OnObjectTransacted.AddLambda([this](UObject*, const FTransactionObjectEvent& Event)
diff --git a/Engine/Source/Editor/ObjectMixer/ObjectMixer/Private/ObjectMixerEditorSerializedData.cpp b/Engine/Source/Editor/ObjectMixer/ObjectMixer/Private/ObjectMixerEditorSerializedData.cpp
--- a/Engine/Source/Editor/ObjectMixer/ObjectMixer/Private/ObjectMixerEditorSerializedData.cpp
+++ b/Engine/Source/Editor/ObjectMixer/ObjectMixer/Private/ObjectMixerEditorSerializedData.cpp
@@ -3,6 +3,11 @@
 #include "ObjectMixerEditorSerializedData.h"
 
 #include "Algo/Find.h"
+#include "Algo/RemoveIf.h"
+#include "ScopedTransaction.h"
+#include "Algo/IndexOf.h"
+
+#define LOCTEXT_NAMESPACE "ObjectMixerEditor"
 
 FObjectMixerSerializationData* UObjectMixerEditorSerializedData::FindSerializationDataByFilterClassName(const FName& FilterClassName)
 {
@@ -20,6 +25,10 @@ FObjectMixerSerializationData* UObjectMixerEditorSerializedData::FindSerializati
 
 void UObjectMixerEditorSerializedData::AddObjectsToCollection(const FName& FilterClassName, const FName& CollectionName, const TSet<FSoftObjectPath>& ObjectsToAdd)
 {
+	FScopedTransaction AddObjectsToCollectionTransaction(LOCTEXT("AddObjectsToCollectionTransaction","Add Objects To Collection"));
+
+	Modify();
+	
 	FObjectMixerSerializationData* Data = FindSerializationDataByFilterClassName(FilterClassName);
 
 	if (!Data)
@@ -51,36 +60,78 @@ void UObjectMixerEditorSerializedData::RemoveObjectsFromCollection(const FName&
 {
 	if (FObjectMixerSerializationData* Data = FindSerializationDataByFilterClassName(FilterClassName))
 	{
-		if (FObjectMixerCollectionObjectSet* Match = Algo::FindByPredicate(Data->SerializedCollection,
+		FScopedTransaction RemoveCollectionTransaction(LOCTEXT("RemoveCollectionTransaction","Remove Collection"));
+	
+		Modify();
+		
+		Data->SerializedCollection.SetNum(Algo::StableRemoveIf(Data->SerializedCollection,
 			[CollectionName](const FObjectMixerCollectionObjectSet& Comparator)
 			{
 				return Comparator.CollectionName.IsEqual(CollectionName);
-			}))
-		{
-			for (const FSoftObjectPath& ObjectPath : ObjectsToRemove)
-			{
-				(*Match).CollectionObjects.Remove(ObjectPath);
-			}
-			
-			SaveConfig();
-		}
+			}));
+		
+		SaveConfig();
 	}
 }
 
 void UObjectMixerEditorSerializedData::RemoveCollection(const FName& FilterClassName, const FName& CollectionName)
-{
+{	
 	if (FObjectMixerSerializationData* Data = FindSerializationDataByFilterClassName(FilterClassName))
 	{
-		if (FObjectMixerCollectionObjectSet* Match = Algo::FindByPredicate(Data->SerializedCollection,
+		FScopedTransaction RemoveCollectionTransaction(LOCTEXT("RemoveCollectionTransaction","Remove Collection"));
+	
+		Modify();
+		
+		Data->SerializedCollection.SetNum(Algo::StableRemoveIf(Data->SerializedCollection,
 			[CollectionName](const FObjectMixerCollectionObjectSet& Comparator)
 			{
 				return Comparator.CollectionName.IsEqual(CollectionName);
+			}));
+		
+		SaveConfig();
+	}
+}
+
+void UObjectMixerEditorSerializedData::ReorderCollection(const FName& FilterClassName,
+	const FName& CollectionToMoveName, const FName& CollectionInsertBeforeName)
+{
+	if (FObjectMixerSerializationData* Data = FindSerializationDataByFilterClassName(FilterClassName))
+	{
+		FObjectMixerCollectionObjectSet MatchCopy;
+		bool bFoundMatch = false;
+		if (const FObjectMixerCollectionObjectSet* Match = Algo::FindByPredicate(Data->SerializedCollection,
+			[CollectionToMoveName](const FObjectMixerCollectionObjectSet& Comparator)
+			{
+				return Comparator.CollectionName.IsEqual(CollectionToMoveName);
 			}))
 		{
-			(*Match).CollectionObjects.Empty();
+			MatchCopy = (*Match);
+			bFoundMatch = true;
+		}
 
-			Data->SerializedCollection.Remove(*Match);
+		if (bFoundMatch)
+		{
+			FScopedTransaction RemoveCollectionTransaction(LOCTEXT("RemoveCollectionTransaction","Remove Collection"));
+	
+			Modify();
+			
+			RemoveCollection(FilterClassName, CollectionToMoveName);
+
+			if (CollectionInsertBeforeName == "All") // Move MatchCopy to the end
+			{
+				Data->SerializedCollection.Add(MatchCopy);
+			}
+			else
+			{
+				const int32 IndexOfCollectionInsertBefore = Algo::IndexOfByPredicate(Data->SerializedCollection,
+					[CollectionInsertBeforeName](const FObjectMixerCollectionObjectSet& Comparator)
+					{
+						return Comparator.CollectionName.IsEqual(CollectionInsertBeforeName);
+					});
 
+				Data->SerializedCollection.Insert(MatchCopy, IndexOfCollectionInsertBefore);
+			}
+		
 			SaveConfig();
 		}
 	}
@@ -123,11 +174,11 @@ TSet<FName> UObjectMixerEditorSerializedData::GetCollectionsForObject(const FNam
 	return {};
 }
 
-TSet<FName> UObjectMixerEditorSerializedData::GetAllCollections(const FName& FilterClassName)
+TArray<FName> UObjectMixerEditorSerializedData::GetAllCollectionNames(const FName& FilterClassName)
 {
 	if (FObjectMixerSerializationData* Data = FindSerializationDataByFilterClassName(FilterClassName))
 	{
-		TSet<FName> Collections;
+		TArray<FName> Collections;
 
 		for (const FObjectMixerCollectionObjectSet& CollectionObjectSet : Data->SerializedCollection)
 		{
@@ -139,3 +190,5 @@ TSet<FName> UObjectMixerEditorSerializedData::GetAllCollections(const FName& Fil
 
 	return {};
 }
+
+#undef LOCTEXT_NAMESPACE
diff --git a/Engine/Source/Editor/ObjectMixer/ObjectMixer/Private/Views/List/ObjectMixerEditorListRow.cpp b/Engine/Source/Editor/ObjectMixer/ObjectMixer/Private/Views/List/ObjectMixerEditorListRow.cpp
--- a/Engine/Source/Editor/ObjectMixer/ObjectMixer/Private/Views/List/ObjectMixerEditorListRow.cpp
+++ b/Engine/Source/Editor/ObjectMixer/ObjectMixer/Private/Views/List/ObjectMixerEditorListRow.cpp
@@ -9,6 +9,24 @@
 #include "Views/List/ObjectMixerEditorList.h"
 #include "Views/List/SObjectMixerEditorList.h"
 
+TSharedRef<FObjectMixerListRowDragDropOp> FObjectMixerListRowDragDropOp::New(const TArray<FObjectMixerEditorListRowPtr>& InItems)
+{
+	check(InItems.Num() > 0);
+
+	TSharedRef<FObjectMixerListRowDragDropOp> Operation = MakeShareable(
+		new FObjectMixerListRowDragDropOp());
+
+	Operation->DraggedItems = InItems;
+
+	Operation->DefaultHoverIcon = FAppStyle::Get().GetBrush("Graph.ConnectorFeedback.Error");
+
+	Operation->DefaultHoverText = NSLOCTEXT("ObjectMixerEditor","DefaultRowHoverText","Drop onto another row or Collection Button.");
+
+	Operation->Construct();
+
+	return Operation;
+}
+
 FObjectMixerEditorListRow::~FObjectMixerEditorListRow()
 {
 	FlushReferences();
diff --git a/Engine/Source/Editor/ObjectMixer/ObjectMixer/Private/Views/List/SObjectMixerEditorListRow.cpp b/Engine/Source/Editor/ObjectMixer/ObjectMixer/Private/Views/List/SObjectMixerEditorListRow.cpp
--- a/Engine/Source/Editor/ObjectMixer/ObjectMixer/Private/Views/List/SObjectMixerEditorListRow.cpp
+++ b/Engine/Source/Editor/ObjectMixer/ObjectMixer/Private/Views/List/SObjectMixerEditorListRow.cpp
@@ -25,50 +25,10 @@
 
 #define LOCTEXT_NAMESPACE "ObjectMixerEditor"
 
-// const FText InsertFormatText = LOCTEXT("InsertAboveFormatText", "Insert {0} {1} {2}");
-// const FText AboveText = LOCTEXT("AboveListItem", "above");
-// const FText BelowText = LOCTEXT("BelowListItem", "below");
-// const FText MultiDragFormatText = LOCTEXT("MultiDragFormatText", "{0} Items");
-
-// class FObjectMixerListRowDragDropOp : public FDecoratedDragDropOp
-// {
-// public:
-// 	DRAG_DROP_OPERATOR_TYPE(FObjectMixerListRowDragDropOp, FDecoratedDragDropOp)
-//
-// 	/** The item being dragged and dropped */
-// 	TArray<FObjectMixerEditorListRowPtr> DraggedItems;
-//
-// 	/** Constructs a new drag/drop operation */
-// 	static TSharedRef<FObjectMixerListRowDragDropOp> New(const TArray<FObjectMixerEditorListRowPtr>& InItems)
-// 	{
-// 		check(InItems.Num() > 0);
-//
-// 		TSharedRef<FObjectMixerListRowDragDropOp> Operation = MakeShareable(
-// 			new FObjectMixerListRowDragDropOp());
-//
-// 		Operation->DraggedItems = InItems;
-//
-// 		Operation->DefaultHoverIcon = FAppStyle::Get().GetBrush("Graph.ConnectorFeedback.Error");
-//
-// 		// Set the display text and the transaction name based on whether we're dragging a single or multiple widgets
-// 		if (InItems.Num() == 1)
-// 		{
-// 			//Operation->DefaultHoverText = FText::FromString(InItems[0]->GetCommandInfo().Pin()->Command);
-// 		}
-// 		else
-// 		{
-// 			Operation->DefaultHoverText =
-// 				FText::Format(
-// 					SObjectMixerEditorListRow::MultiDragFormatText,
-// 					FText::AsNumber(Operation->DraggedItems.Num())
-// 				);
-// 		}
-//
-// 		Operation->Construct();
-//
-// 		return Operation;
-// 	}
-// };
+const FText InsertFormatText = LOCTEXT("InsertAboveFormatText", "Insert {0} {1} {2}");
+const FText AboveText = LOCTEXT("AboveListItem", "above");
+const FText BelowText = LOCTEXT("BelowListItem", "below");
+const FText MultiDragFormatText = LOCTEXT("MultiDragFormatText", "{0} Items");
 
 void SObjectMixerEditorListRow::Construct(
 	const FArguments& InArgs, const TSharedRef<STableViewBase>& InOwnerTable,
@@ -82,10 +42,10 @@ void SObjectMixerEditorListRow::Construct(
 	SMultiColumnTableRow<FObjectMixerEditorListRowPtr>::Construct(
 		FSuperRowType::FArguments()
 		.Padding(1.0f)
-	// 	.OnCanAcceptDrop(this, &SObjectMixerEditorListRow::HandleCanAcceptDrop)
-	// 	.OnAcceptDrop(this, &SObjectMixerEditorListRow::HandleAcceptDrop)
-	// 	.OnDragDetected(this, &SObjectMixerEditorListRow::HandleDragDetected)
-	// 	.OnDragLeave(this, &SObjectMixerEditorListRow::HandleDragLeave)
+		.OnCanAcceptDrop(this, &SObjectMixerEditorListRow::HandleCanAcceptDrop)
+		.OnAcceptDrop(this, &SObjectMixerEditorListRow::HandleAcceptDrop)
+		.OnDragDetected(this, &SObjectMixerEditorListRow::HandleDragDetected)
+		.OnDragLeave(this, &SObjectMixerEditorListRow::HandleDragLeave)
 	 	, InOwnerTable
 	);
 
@@ -160,124 +120,124 @@ SObjectMixerEditorListRow::~SObjectMixerEditorListRow()
 	Item.Reset();
 }
 
-// FReply SObjectMixerEditorListRow::HandleDragDetected(const FGeometry& MyGeometry, const FPointerEvent& MouseEvent)
-// {
-// 	TArray<FObjectMixerEditorListRowPtr> DraggedItems = Item.Pin()->GetSelectedTreeViewItems();
-// 	TSharedRef<FObjectMixerListRowDragDropOp> Operation =
-// 		FObjectMixerListRowDragDropOp::New(DraggedItems);
-//
-// 	return FReply::Handled().BeginDragDrop(Operation);
-// }
-//
-// void SObjectMixerEditorListRow::HandleDragLeave(const FDragDropEvent& DragDropEvent)
-// {
-// 	if (TSharedPtr<FObjectMixerListRowDragDropOp> Operation =
-// 		DragDropEvent.GetOperationAs<FObjectMixerListRowDragDropOp>())
-// 	{
-// 		Operation->ResetToDefaultToolTip();
-// 	}
-// }
-
-// TOptional<EItemDropZone> SObjectMixerEditorListRow::HandleCanAcceptDrop(const FDragDropEvent& DragDropEvent,
-//                                                                              EItemDropZone DropZone,
-//                                                                              FObjectMixerEditorListRowPtr
-//                                                                              TargetItem)
-// {
-// 	TSharedPtr<FObjectMixerListRowDragDropOp> Operation =
-// 		DragDropEvent.GetOperationAs<FObjectMixerListRowDragDropOp>();
-//
-// 	if (!Operation.IsValid())
-// 	{
-// 		return TOptional<EItemDropZone>();
-// 	}
-//
-// 	Operation->SetToolTip(
-// 		LOCTEXT("SortByCustomOrderDrgDropWarning", "Sort by custom order (\"#\") to drag & drop"),
-// 		FAppStyle::Get().GetBrush("Graph.ConnectorFeedback.Error")
-// 	);
-//
-// 	const bool bIsDropDenied = false;
-//
-// 	FString BoolAsString = bIsDropDenied ? "true" : "false";
-//
-// 	if (bIsDropDenied)
-// 	{
-// 		Operation->ResetToDefaultToolTip();
-//
-// 		return TOptional<EItemDropZone>();
-// 	}
-//
-// 	FText ItemNameText = FText::FromString("Light Label");
-//
-// 	if (Operation->DraggedItems.Num() > 1)
-// 	{
-// 		ItemNameText = FText::Format(MultiDragFormatText, FText::AsNumber(Operation->DraggedItems.Num()));
-// 	}
-//
-// 	const FText DropPermittedText =
-// 		FText::Format(InsertFormatText,
-// 		              ItemNameText,
-// 		              DropZone == EItemDropZone::BelowItem ? BelowText : AboveText,
-// 		              ItemNameText
-// 		);
-//
-// 	Operation->SetToolTip(
-// 		DropPermittedText,
-// 		FAppStyle::Get().GetBrush("Graph.ConnectorFeedback.OK")
-// 	);
-//
-// 	// We have no behaviour yet for dropping one item onto another, so we'll treat it like we dropped it above
-// 	return DropZone == EItemDropZone::OntoItem ? EItemDropZone::AboveItem : DropZone;
-// }
-
-// FReply SObjectMixerEditorListRow::HandleAcceptDrop(const FDragDropEvent& DragDropEvent, EItemDropZone DropZone,
-//                                                         FObjectMixerEditorListRowPtr TargetItem)
-// {
-// 	TSharedPtr<FObjectMixerListRowDragDropOp> Operation =
-// 		DragDropEvent.GetOperationAs<FObjectMixerListRowDragDropOp>();
-//
-// 	if (!Operation.IsValid())
-// 	{
-// 		return FReply::Unhandled();
-// 	}
-//
-// 	const TSharedPtr<SObjectMixerEditorList> ListView = Item.Pin()->GetListViewPtr().Pin();
-//
-// 	TArray<FObjectMixerEditorListRowPtr> DraggedItems = Operation->DraggedItems;
-//
-// 	TArray<FObjectMixerEditorListRowPtr> AllTreeItemsCopy = ListView->GetTreeViewItems();
-//
-// 	for (const FObjectMixerEditorListRowPtr& DraggedItem : DraggedItems)
-// 	{
-// 		if (!DraggedItem.IsValid() || !AllTreeItemsCopy.Contains(DraggedItem))
-// 		{
-// 			continue;
-// 		}
-//
-// 		AllTreeItemsCopy.Remove(DraggedItem);
-// 	}
-//
-// 	const int32 TargetIndex = AllTreeItemsCopy.IndexOfByKey(TargetItem);
-//
-// 	if (TargetIndex > -1)
-// 	{
-// 		for (int32 ItemIndex = DraggedItems.Num() - 1; ItemIndex >= 0; ItemIndex--)
-// 		{
-// 			const FObjectMixerEditorListRowPtr& DraggedItem = DraggedItems[ItemIndex];
-//
-// 			if (!DraggedItem.IsValid() || AllTreeItemsCopy.Contains(DraggedItem))
-// 			{
-// 				continue;
-// 			}
-//
-// 			AllTreeItemsCopy.Insert(DraggedItem, DropZone == EItemDropZone::AboveItem ? TargetIndex : TargetIndex + 1);
-// 		}
-//
-// 		ListView->SetTreeViewItems(AllTreeItemsCopy);
-// 	}
-//
-// 	return FReply::Handled();
-// }
+FReply SObjectMixerEditorListRow::HandleDragDetected(const FGeometry& MyGeometry, const FPointerEvent& MouseEvent)
+{
+	TArray<FObjectMixerEditorListRowPtr> DraggedItems = Item.Pin()->GetSelectedTreeViewItems();
+	TSharedRef<FObjectMixerListRowDragDropOp> Operation =
+		FObjectMixerListRowDragDropOp::New(DraggedItems);
+
+	return FReply::Handled().BeginDragDrop(Operation);
+}
+
+void SObjectMixerEditorListRow::HandleDragLeave(const FDragDropEvent& DragDropEvent)
+{
+	if (TSharedPtr<FObjectMixerListRowDragDropOp> Operation =
+		DragDropEvent.GetOperationAs<FObjectMixerListRowDragDropOp>())
+	{
+		Operation->ResetToDefaultToolTip();
+	}
+}
+
+TOptional<EItemDropZone> SObjectMixerEditorListRow::HandleCanAcceptDrop(const FDragDropEvent& DragDropEvent,
+                                                                             EItemDropZone DropZone,
+                                                                             FObjectMixerEditorListRowPtr
+                                                                             TargetItem)
+{
+	TSharedPtr<FObjectMixerListRowDragDropOp> Operation =
+		DragDropEvent.GetOperationAs<FObjectMixerListRowDragDropOp>();
+
+	if (!Operation.IsValid())
+	{
+		return TOptional<EItemDropZone>();
+	}
+
+	Operation->SetToolTip(
+		LOCTEXT("SortByCustomOrderDrgDropWarning", "Sort by custom order (\"#\") to drag & drop"),
+		FAppStyle::Get().GetBrush("Graph.ConnectorFeedback.Error")
+	);
+
+	const bool bIsDropDenied = false;
+
+	FString BoolAsString = bIsDropDenied ? "true" : "false";
+
+	if (bIsDropDenied)
+	{
+		Operation->ResetToDefaultToolTip();
+
+		return TOptional<EItemDropZone>();
+	}
+
+	FText ItemNameText = FText::FromString("Light Label");
+
+	if (Operation->DraggedItems.Num() > 1)
+	{
+		ItemNameText = FText::Format(MultiDragFormatText, FText::AsNumber(Operation->DraggedItems.Num()));
+	}
+
+	const FText DropPermittedText =
+		FText::Format(InsertFormatText,
+		              ItemNameText,
+		              DropZone == EItemDropZone::BelowItem ? BelowText : AboveText,
+		              ItemNameText
+		);
+
+	Operation->SetToolTip(
+		DropPermittedText,
+		FAppStyle::Get().GetBrush("Graph.ConnectorFeedback.OK")
+	);
+
+	// We have no behaviour yet for dropping one item onto another, so we'll treat it like we dropped it above
+	return DropZone == EItemDropZone::OntoItem ? EItemDropZone::AboveItem : DropZone;
+}
+
+FReply SObjectMixerEditorListRow::HandleAcceptDrop(const FDragDropEvent& DragDropEvent, EItemDropZone DropZone,
+                                                        FObjectMixerEditorListRowPtr TargetItem)
+{
+	TSharedPtr<FObjectMixerListRowDragDropOp> Operation =
+		DragDropEvent.GetOperationAs<FObjectMixerListRowDragDropOp>();
+
+	if (!Operation.IsValid())
+	{
+		return FReply::Unhandled();
+	}
+
+	const TSharedPtr<SObjectMixerEditorList> ListView = Item.Pin()->GetListViewPtr().Pin();
+
+	TArray<FObjectMixerEditorListRowPtr> DraggedItems = Operation->DraggedItems;
+
+	TArray<FObjectMixerEditorListRowPtr> AllTreeItemsCopy = ListView->GetTreeViewItems();
+
+	for (const FObjectMixerEditorListRowPtr& DraggedItem : DraggedItems)
+	{
+		if (!DraggedItem.IsValid() || !AllTreeItemsCopy.Contains(DraggedItem))
+		{
+			continue;
+		}
+
+		AllTreeItemsCopy.Remove(DraggedItem);
+	}
+
+	const int32 TargetIndex = AllTreeItemsCopy.IndexOfByKey(TargetItem);
+
+	if (TargetIndex > -1)
+	{
+		for (int32 ItemIndex = DraggedItems.Num() - 1; ItemIndex >= 0; ItemIndex--)
+		{
+			const FObjectMixerEditorListRowPtr& DraggedItem = DraggedItems[ItemIndex];
+
+			if (!DraggedItem.IsValid() || AllTreeItemsCopy.Contains(DraggedItem))
+			{
+				continue;
+			}
+
+			AllTreeItemsCopy.Insert(DraggedItem, DropZone == EItemDropZone::AboveItem ? TargetIndex : TargetIndex + 1);
+		}
+
+		ListView->SetTreeViewItems(AllTreeItemsCopy);
+	}
+
+	return FReply::Handled();
+}
 
 bool SObjectMixerEditorListRow::IsVisible() const
 {
diff --git a/Engine/Source/Editor/ObjectMixer/ObjectMixer/Private/Views/MainPanel/ObjectMixerEditorMainPanel.cpp b/Engine/Source/Editor/ObjectMixer/ObjectMixer/Private/Views/MainPanel/ObjectMixerEditorMainPanel.cpp
--- a/Engine/Source/Editor/ObjectMixer/ObjectMixer/Private/Views/MainPanel/ObjectMixerEditorMainPanel.cpp
+++ b/Engine/Source/Editor/ObjectMixer/ObjectMixer/Private/Views/MainPanel/ObjectMixerEditorMainPanel.cpp
@@ -101,9 +101,9 @@ void FObjectMixerEditorMainPanel::AddObjectsToCollection(const FName& Collection
 {
 	if (const TSubclassOf<UObjectMixerObjectFilter> Filter = GetObjectFilterClass())
 	{
-		if (UObjectMixerEditorSerializedData* Settings = GetMutableDefault<UObjectMixerEditorSerializedData>())
+		if (UObjectMixerEditorSerializedData* SerializedData = GetMutableDefault<UObjectMixerEditorSerializedData>())
 		{
-			Settings->AddObjectsToCollection(Filter->GetFName(), CollectionName, ObjectsToAdd);
+			SerializedData->AddObjectsToCollection(Filter->GetFName(), CollectionName, ObjectsToAdd);
 
 			OnObjectMixerCollectionMapChanged.Broadcast();
 		}
@@ -114,9 +114,9 @@ void FObjectMixerEditorMainPanel::RemoveObjectsFromCollection(const FName& Colle
 {
 	if (const TSubclassOf<UObjectMixerObjectFilter> Filter = GetObjectFilterClass())
 	{
-		if (UObjectMixerEditorSerializedData* Settings = GetMutableDefault<UObjectMixerEditorSerializedData>())
+		if (UObjectMixerEditorSerializedData* SerializedData = GetMutableDefault<UObjectMixerEditorSerializedData>())
 		{
-			Settings->RemoveObjectsFromCollection(Filter->GetFName(), CollectionName, ObjectsToRemove);
+			SerializedData->RemoveObjectsFromCollection(Filter->GetFName(), CollectionName, ObjectsToRemove);
 
 			OnObjectMixerCollectionMapChanged.Broadcast();
 		}
@@ -127,9 +127,22 @@ void FObjectMixerEditorMainPanel::RemoveCollection(const FName& CollectionName)
 {
 	if (const TSubclassOf<UObjectMixerObjectFilter> Filter = GetObjectFilterClass())
 	{
-		if (UObjectMixerEditorSerializedData* Settings = GetMutableDefault<UObjectMixerEditorSerializedData>())
+		if (UObjectMixerEditorSerializedData* SerializedData = GetMutableDefault<UObjectMixerEditorSerializedData>())
 		{
-			Settings->RemoveCollection(Filter->GetFName(), CollectionName);
+			SerializedData->RemoveCollection(Filter->GetFName(), CollectionName);
+
+			OnObjectMixerCollectionMapChanged.Broadcast();
+		}
+	}
+}
+
+void FObjectMixerEditorMainPanel::ReorderCollection(const FName& CollectionToMoveName, const FName& CollectionInsertBeforeName) const
+{
+	if (const TSubclassOf<UObjectMixerObjectFilter> Filter = GetObjectFilterClass())
+	{
+		if (UObjectMixerEditorSerializedData* SerializedData = GetMutableDefault<UObjectMixerEditorSerializedData>())
+		{
+			SerializedData->ReorderCollection(Filter->GetFName(), CollectionToMoveName, CollectionInsertBeforeName);
 
 			OnObjectMixerCollectionMapChanged.Broadcast();
 		}
@@ -140,9 +153,9 @@ bool FObjectMixerEditorMainPanel::IsObjectInCollection(const FName& CollectionNa
 {
 	if (const TSubclassOf<UObjectMixerObjectFilter> Filter = GetObjectFilterClass())
 	{
-		if (UObjectMixerEditorSerializedData* Settings = GetMutableDefault<UObjectMixerEditorSerializedData>())
+		if (UObjectMixerEditorSerializedData* SerializedData = GetMutableDefault<UObjectMixerEditorSerializedData>())
 		{
-			return Settings->IsObjectInCollection(Filter->GetFName(), CollectionName, InObject);
+			return SerializedData->IsObjectInCollection(Filter->GetFName(), CollectionName, InObject);
 		}
 	}
 
@@ -153,22 +166,22 @@ TSet<FName> FObjectMixerEditorMainPanel::GetCollectionsForObject(const FSoftObje
 {
 	if (const TSubclassOf<UObjectMixerObjectFilter> Filter = GetObjectFilterClass())
 	{
-		if (UObjectMixerEditorSerializedData* Settings = GetMutableDefault<UObjectMixerEditorSerializedData>())
+		if (UObjectMixerEditorSerializedData* SerializedData = GetMutableDefault<UObjectMixerEditorSerializedData>())
 		{
-			return Settings->GetCollectionsForObject(Filter->GetFName(), InObject);
+			return SerializedData->GetCollectionsForObject(Filter->GetFName(), InObject);
 		}
 	}
 
 	return {};
 }
 
-TSet<FName> FObjectMixerEditorMainPanel::GetAllCollections() const
+TArray<FName> FObjectMixerEditorMainPanel::GetAllCollectionNames() const
 {
 	if (const TSubclassOf<UObjectMixerObjectFilter> Filter = GetObjectFilterClass())
 	{
-		if (UObjectMixerEditorSerializedData* Settings = GetMutableDefault<UObjectMixerEditorSerializedData>())
+		if (UObjectMixerEditorSerializedData* SerializedData = GetMutableDefault<UObjectMixerEditorSerializedData>())
 		{
-			return Settings->GetAllCollections(Filter->GetFName());
+			return SerializedData->GetAllCollectionNames(Filter->GetFName());
 		}
 	}
 
diff --git a/Engine/Source/Editor/ObjectMixer/ObjectMixer/Private/Views/MainPanel/SObjectMixerEditorMainPanel.cpp b/Engine/Source/Editor/ObjectMixer/ObjectMixer/Private/Views/MainPanel/SObjectMixerEditorMainPanel.cpp
--- a/Engine/Source/Editor/ObjectMixer/ObjectMixer/Private/Views/MainPanel/SObjectMixerEditorMainPanel.cpp
+++ b/Engine/Source/Editor/ObjectMixer/ObjectMixer/Private/Views/MainPanel/SObjectMixerEditorMainPanel.cpp
@@ -7,12 +7,14 @@
 #include "ObjectMixerEditorProjectSettings.h"
 #include "Views/List/ObjectMixerEditorList.h"
 #include "Views/MainPanel/ObjectMixerEditorMainPanel.h"
+#include "Views/Widgets/SCollectionSelectionButton.h"
 #include "Views/Widgets/SObjectMixerPlacementAssetMenuEntry.h"
 
 #include "Framework/MultiBox/MultiBoxBuilder.h"
 #include "Kismet2/SClassPickerDialog.h"
 #include "PlacementMode/Public/IPlacementModeModule.h"
 #include "SPositiveActionButton.h"
+#include "Brushes/SlateRoundedBoxBrush.h"
 #include "Styling/StyleColors.h"
 #include "Views/List/ObjectMixerEditorListFilters/IObjectMixerEditorListFilter.h"
 #include "Views/List/ObjectMixerEditorListFilters/ObjectMixerEditorListFilter_Collection.h"
@@ -478,36 +480,19 @@ void SObjectMixerEditorMainPanel::RebuildCollectionSelector()
 	CollectionSelectorBox->ClearChildren();
 	CollectionSelectorBox->SetVisibility(EVisibility::Collapsed);
 
-	TArray<FName> AllCollections = MainPanelModel.Pin()->GetAllCollections().Array();
+	TArray<FName> AllCollections = MainPanelModel.Pin()->GetAllCollectionNames();
 	
 	if (AllCollections.IsEmpty())
 	{
-		// we've selected something that has no sections - rather than show just "All", hide the box
+		// no collections - rather than show just "All", hide the box
 		ResetCurrentCollectionSelection();
 		return;
 	}
 
 	auto CreateSection = [this](FName CollectionName) 
 	{
-		return SNew(SBox)
-			.Padding(FMargin(0))
-			[
-				SNew(SCheckBox)
-				.Style(FAppStyle::Get(), "DetailsView.SectionButton")
-				.OnCheckStateChanged(this, &SObjectMixerEditorMainPanel::OnCollectionCheckedChanged, CollectionName)
-				.IsChecked(this, &SObjectMixerEditorMainPanel::IsCollectionChecked, CollectionName)
-				[
-					SNew(STextBlock)
-					.TextStyle(FAppStyle::Get(), "SmallText")
-					.Text(FText::FromName(CollectionName))
-				]
-			];
+		return SNew(SCollectionSelectionButton, SharedThis(this), CollectionName);
 	};
-	
-	AllCollections.Sort([](FName A, FName B)
-		{
-			return A.LexicalLess(B);
-		});
 
 	for (const FName& Key : AllCollections)
 	{
@@ -525,16 +510,21 @@ void SObjectMixerEditorMainPanel::RebuildCollectionSelector()
 	CollectionSelectorBox->SetVisibility(EVisibility::Visible);
 }
 
-void SObjectMixerEditorMainPanel::OnCollectionCheckedChanged(ECheckBoxState State, FName SectionName)
+void SObjectMixerEditorMainPanel::RequestRemoveCollection(const FName& CollectionName)
+{
+	MainPanelModel.Pin()->RemoveCollection(CollectionName);
+
+	CurrentCollectionSelection.Remove(CollectionName);
+}
+
+void SObjectMixerEditorMainPanel::OnCollectionCheckedStateChanged(ECheckBoxState State, FName CollectionName)
 {
 	check(MainPanelModel.IsValid());
 
 	// Remove Collection
 	if (FSlateApplication::Get().GetModifierKeys().IsAltDown())
 	{
-		MainPanelModel.Pin()->RemoveCollection(SectionName);
-
-		CurrentCollectionSelection.Remove(SectionName);
+		RequestRemoveCollection(CollectionName);
 	}
 	else
 	{
@@ -544,15 +534,15 @@ void SObjectMixerEditorMainPanel::OnCollectionCheckedChanged(ECheckBoxState Stat
 		{
 			if (bIsControlDown)
 			{
-				CurrentCollectionSelection.Remove(SectionName);
+				CurrentCollectionSelection.Remove(CollectionName);
 			}
 			else
 			{
 				CurrentCollectionSelection.Reset();
 
-				if (SectionName != "All")
+				if (CollectionName != "All")
 				{
-					CurrentCollectionSelection.Add(SectionName);
+					CurrentCollectionSelection.Add(CollectionName);
 				}
 			}
 		}
@@ -563,9 +553,9 @@ void SObjectMixerEditorMainPanel::OnCollectionCheckedChanged(ECheckBoxState Stat
 				CurrentCollectionSelection.Reset();
 			}
 
-			if (SectionName != "All")
+			if (CollectionName != "All")
 			{
-				CurrentCollectionSelection.Add(SectionName);
+				CurrentCollectionSelection.Add(CollectionName);
 			}
 		}
 	}
diff --git a/Engine/Source/Editor/ObjectMixer/ObjectMixer/Private/Views/Widgets/ObjectMixerEditorListMenuContext.cpp b/Engine/Source/Editor/ObjectMixer/ObjectMixer/Private/Views/Widgets/ObjectMixerEditorListMenuContext.cpp
--- a/Engine/Source/Editor/ObjectMixer/ObjectMixer/Private/Views/Widgets/ObjectMixerEditorListMenuContext.cpp
+++ b/Engine/Source/Editor/ObjectMixer/ObjectMixer/Private/Views/Widgets/ObjectMixerEditorListMenuContext.cpp
@@ -108,7 +108,7 @@ void UObjectMixerEditorListMenuContext::RegisterContextMenu()
 							
 							if (const TSharedPtr<FObjectMixerEditorMainPanel> MainPanel = Context->Data.MainPanelPtr.Pin())
 							{
-								if (TArray<FName> Collections = MainPanel->GetAllCollections().Array(); Collections.Num() > 0)
+								if (TArray<FName> Collections = MainPanel->GetAllCollectionNames(); Collections.Num() > 0)
 								{
 									FToolMenuSection& Section = Menu->FindOrAddSection("Collections");
 								   Section.Label = LOCTEXT("CollectionsSectionName", "Collections");
diff --git a/Engine/Source/Editor/ObjectMixer/ObjectMixer/Private/Views/Widgets/SCollectionSelectionButton.cpp b/Engine/Source/Editor/ObjectMixer/ObjectMixer/Private/Views/Widgets/SCollectionSelectionButton.cpp
new file mode 100644
--- /dev/null
+++ b/Engine/Source/Editor/ObjectMixer/ObjectMixer/Private/Views/Widgets/SCollectionSelectionButton.cpp
@@ -0,0 +1,238 @@
+// Copyright Epic Games, Inc. All Rights Reserved.
+
+#pragma once
+
+#include "Views/Widgets/SCollectionSelectionButton.h"
+
+#include "DragAndDrop/DecoratedDragDropOp.h"
+#include "Styling/AppStyle.h"
+#include "Views/MainPanel/ObjectMixerEditorMainPanel.h"
+#include "Views/MainPanel/SObjectMixerEditorMainPanel.h"
+#include "Widgets/Layout/SBorder.h"
+#include "Widgets/Text/STextBlock.h"
+
+#define LOCTEXT_NAMESPACE "ObjectMixerEditor"
+
+class FCollectionSelectionButtonDragDropOp : public FDecoratedDragDropOp
+{
+public:
+	DRAG_DROP_OPERATOR_TYPE(FCollectionSelectionButtonDragDropOp, FDecoratedDragDropOp)
+
+	/** The item being dragged and dropped */
+	FName DraggedItem;
+
+	/** Constructs a new drag/drop operation */
+	static TSharedRef<FCollectionSelectionButtonDragDropOp> New(const FName DraggedItem)
+	{
+		TSharedRef<FCollectionSelectionButtonDragDropOp> Operation = MakeShareable(
+			   new FCollectionSelectionButtonDragDropOp());
+		if (DraggedItem != NAME_None)
+		{
+			Operation->DraggedItem = DraggedItem;
+
+			Operation->DefaultHoverIcon = FAppStyle::Get().GetBrush("Graph.ConnectorFeedback.Error");
+
+			Operation->DefaultHoverText = LOCTEXT("DefaultCollectionButtonHoverText","Drop onto another Collection Button to set a custom order.");
+
+			Operation->Construct();
+		}
+
+		return Operation;
+	}
+};
+
+void SCollectionSelectionButton::Construct(const FArguments& InArgs, const TSharedRef<SObjectMixerEditorMainPanel> MainPanelWidget, const FName& InCollectionName)
+{
+	MainPanelPtr = MainPanelWidget;
+	CollectionName = InCollectionName;
+	
+	ChildSlot
+	[
+		SNew(SBorder)
+		.Padding(FMargin(16, 4))
+		.BorderImage_Lambda([this]()
+		{
+			if (MainPanelPtr.Pin()->IsCollectionChecked(CollectionName) == ECheckBoxState::Checked)
+			{
+				if (bIsPressed)
+				{
+					return &CheckedPressedImage;
+				}
+				else if (bDropIsValid)
+				{
+					return &CheckedValidDropImage;
+				}
+				else if (IsHovered())
+				{
+					return &CheckedHoveredImage;
+				}
+
+				return &CheckedImage;
+			}
+
+			if (bIsPressed)
+			{
+				return &UncheckedPressedImage;
+			}
+			else if (bDropIsValid)
+			{
+				return &UncheckedValidDropImage;
+			}
+			else if (IsHovered())
+			{
+				return &UncheckedHoveredImage;
+			}
+
+			return &UncheckedImage;
+			
+		})
+		[
+			SNew(STextBlock)
+			.TextStyle(FAppStyle::Get(), "SmallText")
+			.Text(FText::FromName(CollectionName))
+			.Visibility(EVisibility::HitTestInvisible)
+		]
+	];
+}
+
+FReply SCollectionSelectionButton::OnMouseButtonDown(const FGeometry& MyGeometry, const FPointerEvent& MouseEvent) 
+{
+	bIsPressed = true;
+	if (MouseEvent.GetEffectingButton() == EKeys::LeftMouseButton )
+	{
+		return FReply::Handled().DetectDrag( SharedThis(this), EKeys::LeftMouseButton );
+	}
+	return FReply::Handled();
+}
+
+FReply SCollectionSelectionButton::OnMouseButtonUp(const FGeometry& MyGeometry, const FPointerEvent& MouseEvent) 
+{
+	bIsPressed = false;
+	if (MouseEvent.GetEffectingButton() == EKeys::LeftMouseButton )
+	{
+		ECheckBoxState CurrentState = MainPanelPtr.Pin()->IsCollectionChecked(CollectionName);
+
+		ECheckBoxState NewState = ECheckBoxState::Checked;
+
+		switch (CurrentState)
+		{
+		case ECheckBoxState::Checked:
+			NewState = ECheckBoxState::Unchecked;
+			break;
+
+		default:
+			break;
+		}
+		
+		MainPanelPtr.Pin()->OnCollectionCheckedStateChanged(NewState, CollectionName);
+	}
+	if (MouseEvent.GetEffectingButton() == EKeys::RightMouseButton)
+	{
+		MainPanelPtr.Pin()->RequestRemoveCollection(CollectionName);
+	}
+	return FReply::Handled();
+}
+
+FReply SCollectionSelectionButton::OnDragDetected(const FGeometry& MyGeometry, const FPointerEvent& MouseEvent)
+{
+	if (CollectionName != "All")
+	{
+		TSharedRef<FCollectionSelectionButtonDragDropOp> OperationFromCollection =
+		   FCollectionSelectionButtonDragDropOp::New(CollectionName);
+
+		OperationFromCollection->ResetToDefaultToolTip();
+
+		bDropIsValid = false;
+
+		return FReply::Handled().BeginDragDrop(OperationFromCollection);
+	}
+
+	return FReply::Handled();
+}
+
+void SCollectionSelectionButton::OnDragEnter(const FGeometry& MyGeometry, const FDragDropEvent& DragDropEvent)
+{
+	if (TSharedPtr<FObjectMixerListRowDragDropOp> OperationFromRow =
+		DragDropEvent.GetOperationAs<FObjectMixerListRowDragDropOp>())
+	{
+		if (CollectionName != "All" && CollectionName != NAME_None)
+		{
+			OperationFromRow->SetToolTip(
+			   LOCTEXT("DropRowItemsOntoCollectionButtonCTA","Add selected items to this collection"),
+			   FAppStyle::Get().GetBrush("Graph.ConnectorFeedback.OK"));
+
+			bDropIsValid = true;
+		}
+	}
+
+	if (TSharedPtr<FCollectionSelectionButtonDragDropOp> OperationFromCollection =
+		DragDropEvent.GetOperationAs<FCollectionSelectionButtonDragDropOp>())
+	{
+		if (CollectionName != NAME_None && CollectionName != OperationFromCollection->DraggedItem)
+		{
+			OperationFromCollection->SetToolTip(
+			   FText::Format(
+				LOCTEXT("DropCollectionButtonOntoCollectionButtonCTA_Format","Reorder {0} before {1}"),
+				FText::FromName(OperationFromCollection->DraggedItem), FText::FromName(CollectionName)),
+			   FAppStyle::Get().GetBrush("Graph.ConnectorFeedback.OK"));
+
+			bDropIsValid = true;
+		}
+	}
+}
+
+void SCollectionSelectionButton::OnDragLeave(const FDragDropEvent& DragDropEvent)
+{
+	bIsPressed = false;
+	bDropIsValid = false;
+	
+	if (TSharedPtr<FObjectMixerListRowDragDropOp> OperationFromRow =
+		DragDropEvent.GetOperationAs<FObjectMixerListRowDragDropOp>())
+	{
+		OperationFromRow->ResetToDefaultToolTip();
+	}
+
+	if (TSharedPtr<FCollectionSelectionButtonDragDropOp> OperationFromCollection =
+		DragDropEvent.GetOperationAs<FCollectionSelectionButtonDragDropOp>())
+	{
+		OperationFromCollection->ResetToDefaultToolTip();
+	}
+}
+
+FReply SCollectionSelectionButton::OnDrop(const FGeometry& MyGeometry, const FDragDropEvent& DragDropEvent)
+{
+	bIsPressed = false;
+	bDropIsValid = false;
+	
+	if (TSharedPtr<FObjectMixerListRowDragDropOp> OperationFromRow =
+		DragDropEvent.GetOperationAs<FObjectMixerListRowDragDropOp>())
+	{
+		if (const TSharedPtr<FObjectMixerEditorMainPanel> MainPanelModel = MainPanelPtr.Pin()->GetMainPanelModel().Pin())
+		{
+			TSet<FSoftObjectPath> ObjectPaths;
+
+			for (const TSharedPtr<FObjectMixerEditorListRow>& Item : OperationFromRow->DraggedItems)
+			{
+				if (const UObject* Object = Item->GetObject())
+				{
+					ObjectPaths.Add(Object);
+				}
+			}
+	
+			MainPanelModel->AddObjectsToCollection(CollectionName, ObjectPaths);
+		}
+	}
+
+	if (TSharedPtr<FCollectionSelectionButtonDragDropOp> OperationFromCollection =
+		DragDropEvent.GetOperationAs<FCollectionSelectionButtonDragDropOp>())
+	{
+		if (const TSharedPtr<FObjectMixerEditorMainPanel> MainPanelModel = MainPanelPtr.Pin()->GetMainPanelModel().Pin())
+		{
+			MainPanelModel->ReorderCollection(OperationFromCollection->DraggedItem, CollectionName);
+		}
+	}
+		
+	return FReply::Handled();
+}
+
+#undef LOCTEXT_NAMESPACE
diff --git a/Engine/Source/Editor/ObjectMixer/ObjectMixer/Public/ObjectMixerEditorSerializedData.h b/Engine/Source/Editor/ObjectMixer/ObjectMixer/Public/ObjectMixerEditorSerializedData.h
--- a/Engine/Source/Editor/ObjectMixer/ObjectMixer/Public/ObjectMixerEditorSerializedData.h
+++ b/Engine/Source/Editor/ObjectMixer/ObjectMixer/Public/ObjectMixerEditorSerializedData.h
@@ -35,7 +35,7 @@ struct FObjectMixerSerializationData
 	FName FilterClassName = NAME_None;
 
 	UPROPERTY()
-	TSet<FObjectMixerCollectionObjectSet> SerializedCollection = {};
+	TArray<FObjectMixerCollectionObjectSet> SerializedCollection = {};
 
 	bool operator==(const FObjectMixerSerializationData& Other) const
 	{
@@ -55,7 +55,10 @@ class OBJECTMIXEREDITOR_API UObjectMixerEditorSerializedData : public UObject
 public:
 	
 	UObjectMixerEditorSerializedData(const FObjectInitializer& ObjectInitializer)
-	{}
+	{
+		// Makes undo/redo possible for this object
+		SetFlags(RF_Transactional);
+	}
 	
 	UPROPERTY(Config)
 	TSet<FObjectMixerSerializationData> SerializedData;
@@ -68,9 +71,11 @@ public:
 
 	void RemoveCollection(const FName& FilterClassName, const FName& CollectionName);
 
+	void ReorderCollection(const FName& FilterClassName, const FName& CollectionToMoveName, const FName& CollectionInsertBeforeName);
+
 	bool IsObjectInCollection(const FName& FilterClassName, const FName& CollectionName, const FSoftObjectPath& InObject);
 
 	TSet<FName> GetCollectionsForObject(const FName& FilterClassName, const FSoftObjectPath& InObject);
 
-	TSet<FName> GetAllCollections(const FName& FilterClassName);
+	TArray<FName> GetAllCollectionNames(const FName& FilterClassName);
 };
diff --git a/Engine/Source/Editor/ObjectMixer/ObjectMixer/Public/Views/List/ObjectMixerEditorListRow.h b/Engine/Source/Editor/ObjectMixer/ObjectMixer/Public/Views/List/ObjectMixerEditorListRow.h
--- a/Engine/Source/Editor/ObjectMixer/ObjectMixer/Public/Views/List/ObjectMixerEditorListRow.h
+++ b/Engine/Source/Editor/ObjectMixer/ObjectMixer/Public/Views/List/ObjectMixerEditorListRow.h
@@ -4,6 +4,7 @@
 
 #include "Containers/Array.h"
 #include "Containers/UnrealString.h"
+#include "DragAndDrop/DecoratedDragDropOp.h"
 #include "Layout/Visibility.h"
 #include "ObjectFilter/ObjectMixerEditorObjectFilter.h"
 #include "Templates/SharedPointer.h"
@@ -18,6 +19,18 @@ struct FSlateBrush;
 struct FObjectMixerEditorListRow;
 typedef TSharedPtr<FObjectMixerEditorListRow> FObjectMixerEditorListRowPtr;
 
+class FObjectMixerListRowDragDropOp : public FDecoratedDragDropOp
+{
+public:
+	DRAG_DROP_OPERATOR_TYPE(FObjectMixerListRowDragDropOp, FDecoratedDragDropOp)
+
+	/** The item being dragged and dropped */
+	TArray<FObjectMixerEditorListRowPtr> DraggedItems;
+
+	/** Constructs a new drag/drop operation */
+	static TSharedRef<FObjectMixerListRowDragDropOp> New(const TArray<FObjectMixerEditorListRowPtr>& InItems);
+};
+
 struct OBJECTMIXEREDITOR_API FObjectMixerEditorListRow final : TSharedFromThis<FObjectMixerEditorListRow>
 {
 	enum EObjectMixerEditorListRowType : uint8
diff --git a/Engine/Source/Editor/ObjectMixer/ObjectMixer/Public/Views/List/SObjectMixerEditorListRow.h b/Engine/Source/Editor/ObjectMixer/ObjectMixer/Public/Views/List/SObjectMixerEditorListRow.h
--- a/Engine/Source/Editor/ObjectMixer/ObjectMixer/Public/Views/List/SObjectMixerEditorListRow.h
+++ b/Engine/Source/Editor/ObjectMixer/ObjectMixer/Public/Views/List/SObjectMixerEditorListRow.h
@@ -32,10 +32,10 @@ public:
 	
 	virtual ~SObjectMixerEditorListRow() override;	
 
-	// FReply HandleDragDetected(const FGeometry& MyGeometry, const FPointerEvent& MouseEvent);
-	// void HandleDragLeave(const FDragDropEvent& DragDropEvent);
-	// TOptional<EItemDropZone> HandleCanAcceptDrop(const FDragDropEvent& DragDropEvent, EItemDropZone DropZone, FObjectMixerEditorListRowPtr TargetItem);
-	// FReply HandleAcceptDrop(const FDragDropEvent& DragDropEvent, EItemDropZone DropZone, FObjectMixerEditorListRowPtr TargetItem);
+	FReply HandleDragDetected(const FGeometry& MyGeometry, const FPointerEvent& MouseEvent);
+	void HandleDragLeave(const FDragDropEvent& DragDropEvent);
+	TOptional<EItemDropZone> HandleCanAcceptDrop(const FDragDropEvent& DragDropEvent, EItemDropZone DropZone, FObjectMixerEditorListRowPtr TargetItem);
+	FReply HandleAcceptDrop(const FDragDropEvent& DragDropEvent, EItemDropZone DropZone, FObjectMixerEditorListRowPtr TargetItem);
 	
 private:
 
diff --git a/Engine/Source/Editor/ObjectMixer/ObjectMixer/Public/Views/MainPanel/ObjectMixerEditorMainPanel.h b/Engine/Source/Editor/ObjectMixer/ObjectMixer/Public/Views/MainPanel/ObjectMixerEditorMainPanel.h
--- a/Engine/Source/Editor/ObjectMixer/ObjectMixer/Public/Views/MainPanel/ObjectMixerEditorMainPanel.h
+++ b/Engine/Source/Editor/ObjectMixer/ObjectMixer/Public/Views/MainPanel/ObjectMixerEditorMainPanel.h
@@ -154,9 +154,10 @@ public:
 	void AddObjectsToCollection(const FName& CollectionName, const TSet<FSoftObjectPath>& ObjectsToAdd) const;
 	void RemoveObjectsFromCollection(const FName& CollectionName, const TSet<FSoftObjectPath>& ObjectsToRemove) const;
 	void RemoveCollection(const FName& CollectionName) const;
+	void ReorderCollection(const FName& CollectionToMoveName, const FName& CollectionInsertBeforeName) const;
 	bool IsObjectInCollection(const FName& CollectionName, const FSoftObjectPath& InObject) const;
 	TSet<FName> GetCollectionsForObject(const FSoftObjectPath& InObject) const;
-	TSet<FName> GetAllCollections() const;
+	TArray<FName> GetAllCollectionNames() const;
 	FOnObjectMixerCollectionMapChanged& GetOnObjectMixerCollectionMapChanged()
 	{
 		return OnObjectMixerCollectionMapChanged;
diff --git a/Engine/Source/Editor/ObjectMixer/ObjectMixer/Public/Views/MainPanel/SObjectMixerEditorMainPanel.h b/Engine/Source/Editor/ObjectMixer/ObjectMixer/Public/Views/MainPanel/SObjectMixerEditorMainPanel.h
--- a/Engine/Source/Editor/ObjectMixer/ObjectMixer/Public/Views/MainPanel/SObjectMixerEditorMainPanel.h
+++ b/Engine/Source/Editor/ObjectMixer/ObjectMixer/Public/Views/MainPanel/SObjectMixerEditorMainPanel.h
@@ -22,6 +22,11 @@ public:
 
 	void Construct(const FArguments& InArgs, const TSharedRef<FObjectMixerEditorMainPanel>& InMainPanel);
 
+	TWeakPtr<FObjectMixerEditorMainPanel> GetMainPanelModel()
+	{
+		return MainPanelModel;
+	}
+
 	FString GetSearchStringFromSearchInputField() const;
 	void SetSearchStringInSearchInputField(const FString InSearchString) const;
 	void ExecuteListViewSearchOnAllRows(const FString& SearchString, const bool bShouldRefreshAfterward = true);
@@ -46,6 +51,11 @@ public:
 	
 	void RebuildCollectionSelector();
 
+	void RequestRemoveCollection(const FName& CollectionName);
+	
+	void OnCollectionCheckedStateChanged(ECheckBoxState State, FName CollectionName);
+	ECheckBoxState IsCollectionChecked(FName Section) const;
+
 	virtual ~SObjectMixerEditorMainPanel() override;
 
 private:
@@ -75,7 +85,4 @@ private:
 	{
 		CurrentCollectionSelection.Reset();
 	}
-	
-	void OnCollectionCheckedChanged(ECheckBoxState State, FName SectionName);
-	ECheckBoxState IsCollectionChecked(FName Section) const;
 };
diff --git a/Engine/Source/Editor/ObjectMixer/ObjectMixer/Public/Views/Widgets/SCollectionSelectionButton.h b/Engine/Source/Editor/ObjectMixer/ObjectMixer/Public/Views/Widgets/SCollectionSelectionButton.h
new file mode 100644
--- /dev/null
+++ b/Engine/Source/Editor/ObjectMixer/ObjectMixer/Public/Views/Widgets/SCollectionSelectionButton.h
@@ -0,0 +1,52 @@
+// Copyright Epic Games, Inc. All Rights Reserved.
+
+#pragma once
+
+#include "Brushes/SlateRoundedBoxBrush.h"
+#include "Styling/StyleColors.h"
+#include "Widgets/SCompoundWidget.h"
+
+#define LOCTEXT_NAMESPACE "ObjectMixerEditor"
+
+class SCollectionSelectionButton final : public SCompoundWidget
+{
+public:
+	
+	SLATE_BEGIN_ARGS(SCollectionSelectionButton)
+	{}
+
+	SLATE_END_ARGS()
+	
+	void Construct(const FArguments& InArgs, const TSharedRef<class SObjectMixerEditorMainPanel> MainPanelWidget, const FName& InCollectionName);
+
+	virtual FReply OnMouseButtonDown(const FGeometry& MyGeometry, const FPointerEvent& MouseEvent) override;
+
+	virtual FReply OnMouseButtonUp(const FGeometry& MyGeometry, const FPointerEvent& MouseEvent) override;
+
+	virtual FReply OnDragDetected(const FGeometry& MyGeometry, const FPointerEvent& MouseEvent) override;
+
+	virtual void OnDragEnter(const FGeometry& MyGeometry, const FDragDropEvent& DragDropEvent) override;
+
+	virtual void OnDragLeave(const FDragDropEvent& DragDropEvent) override;
+
+	virtual FReply OnDrop(const FGeometry& MyGeometry, const FDragDropEvent& DragDropEvent) override;
+
+	TWeakPtr<SObjectMixerEditorMainPanel> MainPanelPtr;
+	FName CollectionName = NAME_None;
+
+private:
+
+	bool bIsPressed = false;
+	bool bDropIsValid = false;
+
+	FSlateRoundedBoxBrush UncheckedImage = FSlateRoundedBoxBrush(FStyleColors::Header, 4.0f, FStyleColors::Input, 1.0f);
+	FSlateRoundedBoxBrush UncheckedHoveredImage = FSlateRoundedBoxBrush(FStyleColors::Hover, 4.0f, FStyleColors::Input, 1.0f);
+	FSlateRoundedBoxBrush UncheckedPressedImage = FSlateRoundedBoxBrush(FStyleColors::Hover, 4.0f, FStyleColors::Input, 1.0f);
+	FSlateRoundedBoxBrush UncheckedValidDropImage = FSlateRoundedBoxBrush(FStyleColors::Secondary, 4.0f, FStyleColors::Input, 1.0f);
+	FSlateRoundedBoxBrush CheckedImage = FSlateRoundedBoxBrush(FStyleColors::Primary, 4.0f, FStyleColors::Input, 1.0f);
+	FSlateRoundedBoxBrush CheckedHoveredImage = FSlateRoundedBoxBrush(FStyleColors::PrimaryHover, 4.0f, FStyleColors::Input, 1.0f);
+	FSlateRoundedBoxBrush CheckedPressedImage = FSlateRoundedBoxBrush(FStyleColors::PrimaryHover, 4.0f, FStyleColors::Input, 1.0f);
+	FSlateRoundedBoxBrush CheckedValidDropImage = FSlateRoundedBoxBrush(FStyleColors::Secondary, 4.0f, FStyleColors::Input, 1.0f);
+};
+
+#undef LOCTEXT_NAMESPACE
diff --git a/Engine/Source/Editor/ObjectMixer/ObjectMixer/Public/Views/Widgets/SObjectMixerEditorTooltipWidget.h b/Engine/Source/Editor/ObjectMixer/ObjectMixer/Public/Views/Widgets/SObjectMixerEditorTooltipWidget.h
deleted file mode 100644
